Fix for build break in lightmass plugin.

diff --git a/Engine/Plugins/Experimental/GPULightmass/Source/GPULightmass/Private/Scene/Scene.cpp b/Engine/Plugins/Experimental/GPULightmass/Source/GPULightmass/Private/Scene/Scene.cpp
--- a/Engine/Plugins/Experimental/GPULightmass/Source/GPULightmass/Private/Scene/Scene.cpp
+++ b/Engine/Plugins/Experimental/GPULightmass/Source/GPULightmass/Private/Scene/Scene.cpp
@@ -795,7 +795,6 @@ void FScene::AddGeometryInstanceFromComponent(UStaticMeshComponent* InComponent)
 			false,
 			false,
 			0b111,
-			1.0f,
 			0,
 			INDEX_NONE,
 			false,
@@ -1069,7 +1068,6 @@ void FScene::AddGeometryInstanceFromComponent(UInstancedStaticMeshComponent* InC
 			false,
 			false,
 			0b111,
-			1.0f,
 			0,
 			INDEX_NONE,
 			false,
@@ -1359,7 +1357,6 @@ void FScene::AddGeometryInstanceFromComponent(ULandscapeComponent* InComponent)
 			false,
 			false,
 			0b111,
-			1.0f,
 			0,
 			INDEX_NONE,
 			false,
