Adopt the newest path tracer changes

diff --git a/Engine/Plugins/Experimental/GPULightmass/Shaders/Private/LightmapPathTracing.usf b/Engine/Plugins/Experimental/GPULightmass/Shaders/Private/LightmapPathTracing.usf
--- a/Engine/Plugins/Experimental/GPULightmass/Shaders/Private/LightmapPathTracing.usf
+++ b/Engine/Plugins/Experimental/GPULightmass/Shaders/Private/LightmapPathTracing.usf
@@ -125,6 +125,22 @@ FThreeBandSHVectorFloat SHBasisFunction3Float(float3 InputVector)
 	#define USE_IRRADIANCE_CACHING 0
 #endif
 
+void AccumulateRadiance(inout float3 TotalRadiance, float3 PathRadiance)
+{
+	float MaxPathIntensity = 0;
+	if (MaxPathIntensity > 0)
+	{
+		// User asked for path contributions to be clamped to reduce fireflies.
+		// Depending on how aggressive this value is, the image could be quite biased
+		TotalRadiance += min(PathRadiance, MaxPathIntensity);
+	}
+	else
+	{
+		// Just average values directly
+		TotalRadiance += PathRadiance;
+	}
+}
+
 void PathTracingKernel(
 	in uint RenderPassIndex,
 	float3 WorldPosition,
@@ -144,132 +160,138 @@ void PathTracingKernel(
 #if USE_IRRADIANCE_CACHING
 	bool bShouldEmitGeometryHitPoint = false;
 	bool bShouldAccumulateToCacheRecord = false;
+	
+	FIrradianceCacheRecord IrradianceCachePrimaryGIHitPoint = (FIrradianceCacheRecord)0;	
+	float RadianceProbe_Pdf = 1.0f / (2 * PI);
+	
 	uint NewEntrySize = 32;
 	int NearestCacheEntryID = -1;
 #endif
+
+	float3 Radiance = 0;
 	
-	float3 Irradiance = 0;
-	
-	float3 RayThroughput = 1.0;
-	float RayPdf = 1.0;
-	
+	// GPULightmass's SampleEmitter(): generate a fake camera ray hitting the texel
 	RayDesc Ray;
 	Ray.Origin = WorldPosition + ShadingNormal;
 	Ray.Direction = -ShadingNormal;
 	Ray.TMin = -0.01f;
 	Ray.TMax = 1.01f;
 	
-	FMaterialClosestHitPayload Payload = (FMaterialClosestHitPayload)0;
-	Payload.RayCone = (FRayCone)0;
-	Payload.WorldPos = WorldPosition;
-	Payload.WorldNormal = ShadingNormal;
-	Payload.Radiance = float3(0, 0, 0);
-	Payload.BaseColor = float3(1, 1, 1);
-	Payload.DiffuseColor = float3(1, 1, 1);
-	Payload.SpecularColor = float3(0, 0, 0);
-	Payload.Specular = 0;
-	Payload.Roughness = 1;
-	Payload.Metallic = 0;
-	Payload.Ior = 1;
-	Payload.CustomData = float4(0, 0, 0, 0);
-	Payload.Opacity = 1;
-	Payload.ShadingModelID = SHADINGMODELID_DEFAULT_LIT;
-	Payload.BlendingMode = RAY_TRACING_BLEND_MODE_OPAQUE;
-	Payload.PrimitiveLightingChannelMask = 0b111;
-	Payload.HitT = 1.0f;
-	Payload.SetFrontFace();
-	uint RayFlags = 0;
-	FRayCone RayCone = (FRayCone)0;
+	// This array will hold a CDF for light picking
+	// Seed the array with a uniform CDF at first so that we always have a valid CDF
+	float LightPickingCdf[RAY_TRACING_LIGHT_COUNT_MAXIMUM];
 
-	float MaterialPdf = 0.0;
-	FMaterialClosestHitPayload PrevMaterialPayload = (FMaterialClosestHitPayload)0;
-	uint Bounces = 0;
-	bool bIsNeeValid = true;
+	InitLightPickingCdf(LightPickingCdf);
 	
-	FIrradianceCacheRecord IrradianceCachePrimaryGIHitPoint = (FIrradianceCacheRecord)0;	
-	float RadianceProbe_Pdf = 1.0f / (2 * PI);
+	FRayCone RayCone = (FRayCone)0;
+	
+	// Cast ray
+	const uint InitialRayFlags = 0;
+	const bool bInitialDisableSkyLightContribution = false;
+	const bool bIgnoreTranslucentMaterials = true;
+
+	// path state variables (these cary information between bounces)
+	float3 PathThroughput = 1.0;
+
+	// We don't want the camera ray to have any MIS contribution against direct light sampling so we pretend that it came from
+	// a singular Material scattering event.
+	float PrevMaterialPdf = POSITIVE_INFINITY;
+
+	// 0: only Material sampling
+	// 1: only Light sampling
+	// 2: both Material and Light
+	const int MISMode = 2; //PathTracingData.MISMode;
 	
-	while (Bounces < 32)
+	for (int Bounce = 0; Bounce <= 32; Bounce++)
 	{
-		// Environment contribution
-		if (Payload.IsMiss())
+		FMaterialClosestHitPayload Payload = (FMaterialClosestHitPayload)0;
+		
+		if (Bounce == 0)
+		{
+			// GPULightmass fakes a hit result on the texel
+			Payload.RayCone = (FRayCone)0;
+			Payload.WorldPos = WorldPosition;
+			Payload.WorldNormal = ShadingNormal;
+			Payload.Radiance = float3(0, 0, 0);
+			Payload.BaseColor = float3(1, 1, 1);
+			Payload.DiffuseColor = float3(1, 1, 1);
+			Payload.SpecularColor = float3(0, 0, 0);
+			Payload.Specular = 0;
+			Payload.Roughness = 1;
+			Payload.Metallic = 0;
+			Payload.Ior = 1;
+			Payload.CustomData = float4(0, 0, 0, 0);
+			Payload.Opacity = 1;
+			Payload.ShadingModelID = SHADINGMODELID_DEFAULT_LIT;
+			Payload.BlendingMode = RAY_TRACING_BLEND_MODE_OPAQUE;
+			Payload.PrimitiveLightingChannelMask = 0b111;
+			Payload.HitT = 1.0f;
+			Payload.SetFrontFace();
+		}
+		else
+		{
+			Ray.TMin = 0.01;
+			Ray.TMax = 1e20;
+		
+			Payload = TraceMaterialRay(
+				TLAS,
+				InitialRayFlags,
+				RAY_TRACING_MASK_ALL,
+				Ray,
+				RayCone,
+				uint2(0, 0),
+				bInitialDisableSkyLightContribution,
+				bIgnoreTranslucentMaterials);
+		}
+		
+		if (MISMode != 1 && Bounce != 0) // In GPULightmass it is not possible to hit any light on its first fake ray
 		{
-			// WARNING: Until ray payload encodes NEE LightId, the execution depends on this ordering!
-			uint SkyLightId = 0;
-			float3 EnvironmentRadiance = 0.0;
-			SkyLight_EvalLight(SkyLightId, Ray.Direction, Ray, EnvironmentRadiance);
-			if (length(EnvironmentRadiance) > 0.0)
+			float HitT = Payload.IsMiss() ? Ray.TMax : Payload.HitT;
+			// Loop over lights to capture their contribution
+			// #dxr_todo: if we have lots of lights, having some hierarchical structure would be better ....
+			float PreviousCdfValue = 0.0;
+			for (uint LightId = 0; LightId < SceneLightsData.Count; ++LightId)
 			{
-				if (IsUnidirectionalEnabled)
-				{
-					// Apply NEE Pdf for correct MIS weight
-					float MisWeight = 1.0;
-					if (IsNextEventEstimationEnabled && Bounces > 0 && bIsNeeValid)
-					{
-						float NeePdf = 0.0;
-						float3 LightUV = GetLightUV(SkyLightId, Ray, Payload);
-						PdfLight(Ray, PrevMaterialPayload, SkyLightId, LightUV, NeePdf);
-	
-						// Note: MaterialPdf will need to be stored on the ray for an iterative solution
-						MisWeight = MaterialPdf / (MaterialPdf + NeePdf);
-					}
-					
-					if (Bounces > 1 || SceneLightsData.Mobility[SkyLightId] == 0)
-					{
-						Irradiance += EnvironmentRadiance * RayThroughput * MisWeight / RayPdf;
-					}
-				}
-				// Record primary environment hits with NEE
-				else if (IsNextEventEstimationEnabled && Bounces == 0)
+				// Camera rays should only do the rest of the loop if we have visible lights OR if
+				// we have a skydome (which are always visible)
+				if (Bounce == 0 && !IsEnvironmentLight(LightId))
+					continue;
+
+				float4 LightResult = TraceLight(Ray, LightId, HitT);
+				float3 LightRadiance = LightResult.xyz;
+				float LightPdf = LightResult.w;
+
+				float LightPickPdf = LightPickingCdf[LightId] - PreviousCdfValue;
+				PreviousCdfValue = LightPickingCdf[LightId];
+
+				// GPULightmass records direct contribution for only static lights
+				if (Bounce > 1 || (Bounce == 1 && SceneLightsData.Mobility[LightId] == 0))
 				{
-					if (Bounces > 1 || SceneLightsData.Mobility[SkyLightId] == 0)
-					{
-						Irradiance += EnvironmentRadiance * RayThroughput / RayPdf;
-					}
+					AccumulateRadiance(Radiance, PathThroughput * LightRadiance * MISWeightRobust(PrevMaterialPdf, LightPdf * LightPickPdf));
 				}
 			}
-			break;
 		}
 		
-		if (Payload.IsHit() && !Payload.IsFrontFace() && Payload.ShadingModelID != SHADINGMODELID_TWOSIDED_FOLIAGE)
+		if (Payload.IsMiss())
 		{
-			bIsValidSample = false;
+			// Ray didn't hit any real geometry, so nothing left to do
 			break;
 		}
-	
-		// Emissive geometry contribution
-		float3 Emission = GetMaterialEmission(Payload);
-		if (any(Emission) > 0.0)
+		
+		if (!Payload.IsFrontFace() && Payload.ShadingModelID != SHADINGMODELID_TWOSIDED_FOLIAGE)
 		{
-			if (IsUnidirectionalEnabled)
-			{
-				// Apply NEE Pdf for correct MIS weight
-				float MisWeight = 1.0;
-				if (IsNextEventEstimationEnabled && Bounces > 0 && bIsNeeValid)
-				{
-					float NeePdf = 0.0;
-					int LightId = GetLightId(Payload);
-					float3 LightUV = GetLightUV(LightId, Ray, Payload);
-					PdfLight(Ray, PrevMaterialPayload, LightId, LightUV, NeePdf);
-	
-					float MisWeight = MaterialPdf / (MaterialPdf + NeePdf);
-				}
-	
-				Irradiance += Emission * RayThroughput * MisWeight / RayPdf;
-			}
-			// Record primary environment hits with NEE
-			else if (IsNextEventEstimationEnabled && Bounces == 0)
-			{
-				Irradiance += Emission * RayThroughput / RayPdf;
-			}
+			bIsValidSample = false;
 			break;
 		}
+	
+		// add in surface emission
+		AccumulateRadiance(Radiance, PathThroughput * Payload.Radiance);
 		
 	#if USE_IRRADIANCE_CACHING
 		uint CoarseLevelSpacing = IrradianceCachingParameters.Spacing;
 		uint DetailLevelSpacing = IrradianceCachingParameters.Spacing;
 		
-		if (Bounces >= 1 && Payload.HitT >= DetailLevelSpacing * IrradianceCachingParameters.CornerRejection && !bShouldAccumulateToCacheRecord)
+		if (Bounce >= 1 && Payload.HitT >= DetailLevelSpacing * IrradianceCachingParameters.CornerRejection && !bShouldAccumulateToCacheRecord)
 		{
 			bool bIrradianceQuerySuccessful = false;
 			bool bGeometryQuerySuccessful = false;
@@ -306,13 +328,13 @@ void PathTracingKernel(
 			if (bIrradianceQuerySuccessful)
 			{
 				// Successful query. Terminate path immediately. Can happen when accumulating for IC entries.
-				Irradiance += RecordIrradiance * RayThroughput / RayPdf;
+				Radiance += RecordIrradiance * PathThroughput;
 				break;
 			}
 			else
 			{
 				// Failed query.
-				if (Bounces == 1)
+				if (Bounce == 1)
 				{
 					if (!bGeometryQuerySuccessful)
 					{
@@ -331,182 +353,182 @@ void PathTracingKernel(
 		}
 	#endif
 	
-		Bounces = Bounces + 1;
 		// Update origin
-		Ray.Origin = Ray.Origin + Ray.Direction * Payload.HitT;
-		FMaterialClosestHitPayload RayHitInfo = Payload; 
-	
-		bIsNeeValid = IsNeeValidMaterial(RayHitInfo);
+		Ray.Origin = Payload.WorldPos;
 
-		// Perform next-event estimation
-		if (IsNextEventEstimationEnabled && bIsNeeValid)
+		// Make sure the surface normal points toward the viewer
+		if (dot(Ray.Direction, Payload.WorldNormal) > 0.0)
 		{
-			int LightId;
-			float3 LightUV;
-			float NeePdf = 0.0;
-			float4 RandSample = RandomSequence_GenerateSample4D(RandSequence, SampleIndex);
-			SampleLight(Ray, RayHitInfo, RandSample, LightId, LightUV, NeePdf);
-	
-			// Fire visibility ray
-			if (NeePdf > 0.0)
+			bIsValidSample = false;
+			break;
+		}
+
+		// Choose a random number for both Light sampling and BxDF sampling
+		float4 RandSample = RandomSequence_GenerateSample4D(RandSequence, SampleIndex);
+
+		// Does this material require NEE?
+		bool bIsNeeValid = IsNeeValidMaterial(Payload);
+		
+		// If we are using Light sampling and we are not on the last bounce ....
+		if (MISMode != 0 && Bounce < 32 && bIsNeeValid && SceneLightsData.Count > 0)
+		{
+			// Choose a light and sample it
+
+			float3 WorldPos = Payload.WorldPos;
+			float3 WorldNormal = Payload.WorldNormal;
+			uint PrimitiveLightingChannelMask = Payload.PrimitiveLightingChannelMask;
+
+			bool IsTransmissiveMaterial = false; // ENABLE_TRANSMISSION && IsMaterialTransmissive(Payload);
+			
+			if (InitLightPickingCdf(WorldPos, WorldNormal, PrimitiveLightingChannelMask, IsTransmissiveMaterial, LightPickingCdf))
 			{
-				RayDesc LightRay;
-	
-				GenerateLightRay(Ray, LightId, LightUV, LightRay);
-				ApplyPositionBias(LightRay, Payload.WorldNormal, 0.01f);
+				// init worked
+				int LightId;
+				float LightPickPdf = 0;
 
-				const uint RayFlags = RAY_FLAG_ACCEPT_FIRST_HIT_AND_END_SEARCH | RAY_FLAG_SKIP_CLOSEST_HIT_SHADER;
-				const uint InstanceInclusionMask = RAY_TRACING_MASK_SHADOW;
+				SelectLight(RandSample.x, LightPickingCdf, LightId, LightPickPdf);
 
-				FMinimalPayload NeePayload = TraceVisibilityRay(
-					TLAS,
-					RayFlags,
-					InstanceInclusionMask,
-					uint2(0, 0),
-					LightRay);
-	
-				if (NeePayload.IsMiss())
+				RayDesc LightRay;
+				LightRay.Origin = WorldPos;
+				LightRay.TMin = 0;
+				float3 RadianceOverPdf;
+				float LightPdf;
+				SampleLight(LightId, RandSample, WorldPos, WorldNormal,
+					LightRay.Direction,
+					LightRay.TMax,
+					RadianceOverPdf,
+					LightPdf);
+				RadianceOverPdf /= LightPickPdf;
+				LightPdf *= LightPickPdf;
+
+				if (LightPdf > 0)
 				{
-					FMaterialClosestHitPayload RayHitInfo = Payload; 
-				#if 0
-					// DEBUG: Validate NEE PDF
-					float NeePdfTmp;
-					PdfLight(Ray, RayHitInfo, LightId, LightUV, NeePdfTmp);
-					if (abs(NeePdf - NeePdfTmp) > 0.0001)
+					float SignedPositionBias = 1.0;
+					if (ENABLE_TRANSMISSION && IsMaterialTransmissive(Payload))
 					{
-						Irradiance = float3(-10, -10, -10);
-						break;
+						// for transmissive materials, bias the position to the other side of the surface if the light is coming from behind
+						float NoL = dot(Payload.WorldNormal, LightRay.Direction);
+						SignedPositionBias = sign(NoL);
 					}
-				#endif
-					// Evaluate radiance
-					float3 Radiance;
-					EvalLight(LightId, LightUV, LightRay, Radiance);
-	
-					// Evaluate material
-					float3 MaterialEvalWeight = 0;
-					float MaterialEvalPdf = 0.0;
-					EvalMaterial(Ray.Direction, LightRay.Direction, RayHitInfo, Bounces == 1, MaterialEvalWeight, MaterialEvalPdf);
-					float3 MaterialThroughput = MaterialEvalWeight * MaterialEvalPdf;
-	
-					// Apply material Pdf for correct MIS weight
-					float MisWeight = 1.0;
-					if (IsUnidirectionalEnabled && IsPhysicalLight(LightId))
+					ApplyPositionBias(LightRay, SignedPositionBias * Payload.WorldNormal, 0.01f);
+
+
+					const uint RayFlags = RAY_FLAG_ACCEPT_FIRST_HIT_AND_END_SEARCH | RAY_FLAG_SKIP_CLOSEST_HIT_SHADER;
+					const uint InstanceInclusionMask = RAY_TRACING_MASK_ALL;
+
+					FPackedMaterialClosestHitPayload PackedPayload = (FPackedMaterialClosestHitPayload)0;
+
+					TraceVisibilityRayPacked(PackedPayload, TLAS, RayFlags, InstanceInclusionMask, uint2(0, 0), LightRay);
+
+					// #dxr_todo: Is it cheaper to fire the ray first? Or eval the material first?
+					if (PackedPayload.IsMiss())
 					{
-					#if 0
-						// DEBUG: Validate Material PDF
-						float MaterialPdfTmp;
-						PdfMaterial(Ray.Direction, LightRay, RayHitInfo, Bounces == 1, MaterialPdfTmp);
-						if (abs(MaterialPdf - MaterialPdfTmp) > 0.0001)
+						// Evaluate material
+						float3 MaterialWeight;
+						float MaterialPdf = 0.0;
+						EvalMaterial(Ray.Direction, LightRay.Direction, Payload, Bounce == 0, MaterialWeight, MaterialPdf);
+
+						float3 MaterialEval = MaterialWeight * MaterialPdf;
+
+						float MISWeight = MISMode == 2 ? MISWeightRobust(LightPdf, MaterialPdf) : 1.0;
+
+						// Record the contribution
+						if (Bounce > 0)
 						{
-							MaterialThroughput *= -1.f;
+							AccumulateRadiance(Radiance, MISWeight * PathThroughput * RadianceOverPdf * MaterialEval);
+						}
+						else
+						{
+							// GPULightmass records contribution of direct lighting separately, and only for static lights
+							if (SceneLightsData.Mobility[LightId] == 0)
+							{
+								DirectLightingNEERadianceValue += MISWeight * PathThroughput * RadianceOverPdf * MaterialEval;
+								DirectLightingNEERadianceDirection = LightRay.Direction;
+							}
 						}
-					#endif
-	
-						MisWeight = NeePdf / (NeePdf + MaterialEvalPdf);
-					}
-	
-					// Record the contribution
-					if (Bounces == 1 && SceneLightsData.Mobility[LightId] == 0)
-					{
-						DirectLightingNEERadianceValue += Radiance * MaterialThroughput * RayThroughput * MisWeight / (NeePdf * RayPdf);
-						DirectLightingNEERadianceDirection = LightRay.Direction;
-					}
-					else if (Bounces > 1 || SceneLightsData.Mobility[LightId] == 0)  
-					{
-						Irradiance += Radiance * MaterialThroughput * RayThroughput * MisWeight / (NeePdf * RayPdf);
 					}
 				}
 			}
 		}
+
+		if (Bounce == 32)
+		{
+			// we've arrived at the last bounce, don't bother trying to sample a direction
+			break;
+		}
 	
 		// Sample material
 		float3 Direction;
-		float3 Weight;
-		float4 RandSample = RandomSequence_GenerateSample4D(RandSequence, SampleIndex);
+		float3 MaterialWeight;
 		float SignedPositionBias;
-		float SamplePdf;
-		SampleMaterial(Ray.Direction, RayHitInfo, RandSample, Bounces == 1, Direction, Weight, SamplePdf, SignedPositionBias);
-		if (SamplePdf == 0.0)
+		float MaterialSamplePdf;
+		SampleMaterial(Ray.Direction, Payload, RandSample, Bounce == 0, Direction, MaterialWeight, MaterialSamplePdf, SignedPositionBias);
+
+		if (MaterialSamplePdf < 0 || asuint(MaterialSamplePdf) > 0x7F800000)
 		{
+			// Pdf became invalid (either negative or NaN)
+			bIsValidSample = false;
 			break;
 		}
-		float3 Throughput = Weight * SamplePdf;
-	
-		PdfMaterial(Ray.Direction, Direction, Payload, Bounces == 1, MaterialPdf);
-		// Debug: Evaluation of a candidate sample direction should maintain energy balance with pathspace MIS
-		// EvalMaterial(Ray.Direction, Direction, RayHitInfo, Bounces == 1, Throughput, MaterialPdf);
-		// SamplePdf = MaterialPdf;
 
-		// Update ray
-		Ray.Direction = Direction;
-		ApplyPositionBias(Ray, SignedPositionBias * Payload.WorldNormal, 0.01f);
-		RayThroughput *= Throughput;
-		RayPdf *= SamplePdf;
-	
-		// Russian roulette
-		float PathEnergy = max(max(RayThroughput.x, RayThroughput.y), RayThroughput.z);
-		float EnergyThreshold = 1.0; 
-		if (PathEnergy < EnergyThreshold && 
-		#if !USE_IRRADIANCE_CACHING
-			Bounces > 4
-		#else
-			// A starting bounce of 4 is too early for relatively high reflectance scenes (0.8), which generates extreme variance that gets rejected by firefly rejection
-			// and shows as missing energy. IC effectively extends the path so it is less subject to this effect. However due to retrace around corners it shows as contrast between cached and non-cached paths.
-			// delay RR to a higher bounce numbmer
-			Bounces > 16
-		#endif
-		)
+		if (!(MaterialSamplePdf > 0))
+		{
+			// No valid direction -- we are done
+			break;
+		}
+
+		float3 NextPathThroughput = PathThroughput * MaterialWeight;
+		if (!any(NextPathThroughput > 0))
 		{
-			float Q = 1.0 - PathEnergy;
-			if (RandomSequence_GenerateSample1D(RandSequence, SampleIndex) < Q)
+			// no energy left in this path
+			break;
+		}
+
+		// Russian roulette:
+		//   The probability of keeping the path should be roughly proportional to the weight at the current shade point,
+		//  but just using MaterialWeight would miss out on cases where the path throughput changes color (like in a cornell
+		//  box when bouncing between walls of different colors). So use the ratio of the brightest color channel in the
+		//  previous and next throughput.
+		//   The second tweak is to add a sqrt() around the probability to soften the termination probability (paths will last
+		//  a little longer). This allows paths to go a bit deeper than the naive heuristic while still allowing them to terminate
+		//  early. This makes RR effective from the very first bounce without needing to delay it.
+		float ContinuationProb = sqrt(saturate(max(NextPathThroughput.x, max(NextPathThroughput.y, NextPathThroughput.z)) / max(PathThroughput.x, max(PathThroughput.y, PathThroughput.z))));
+		if (ContinuationProb < 1)
+		{
+			// If there is some chance we should terminate the ray, draw an extra random value
+			float RussianRouletteRand = RandSample.w; // SampleMaterial does not use this value at the moment
+			//RussianRouletteRand = RandomSequence_GenerateSample1D(RandSequence, SampleIndex);
+			if (RussianRouletteRand >= ContinuationProb)
 			{
+				// stochastically terminate the path
 				break;
 			}
-	
-			RayThroughput /= (1.0 - Q);
+			PathThroughput = NextPathThroughput / ContinuationProb;
 		}
-	
-		// Firefly rejection
-		PathEnergy = max(max(RayThroughput.x, RayThroughput.y), RayThroughput.z);
-		float FireflyThreshold = 50.0;
-		if (PathEnergy / RayPdf > FireflyThreshold)
+		else
 		{
-			break;
+			PathThroughput = NextPathThroughput;
 		}
-	
-		if (dot(Ray.Direction, Payload.WorldNormal) < 0) 
+
+		if (MISMode == 2)
 		{
-			bIsValidSample = false;
-			break;
+			// If we are doing MIS, the Material ray on the next loop through will need to know the Pdf
+			PrevMaterialPdf = MaterialSamplePdf;
 		}
+
+		// Update ray and loop back around to trace again
+		Ray.Direction = Direction;
+		Ray.TMin = 0;
+		ApplyPositionBias(Ray, SignedPositionBias* Payload.WorldNormal, 0.01f);
 		
-		if (Bounces == 1)
+		if (Bounce == 0)
 		{
 			RadianceDirection = Ray.Direction;
-			RadianceProbe_Pdf = SamplePdf;
+		#if USE_IRRADIANCE_CACHING
+			RadianceProbe_Pdf = MaterialSamplePdf;
+		#endif
 			PrimaryRandSample = RandSample;
 		}
-	
-		// Trace again
-		//#dxr-todo: check ray flags when hit shaders are supported
-		Ray.TMin = 0.01;
-		Ray.TMax = 1e20;
-		
-		const uint RayFlags = 0; //RAY_FLAG_CULL_BACK_FACING_TRIANGLES;
-		const uint InstanceInclusionMask = RAY_TRACING_MASK_ALL;
-		const bool bDisableSkyLightContribution = false;
-		const bool bIgnoreTranslucentMaterials = true;
-
-		PrevMaterialPayload = Payload;
-		Payload = TraceMaterialRay(
-			TLAS,
-			RayFlags,
-			InstanceInclusionMask,
-			Ray,
-			RayCone,
-			uint2(0, 0),
-			bDisableSkyLightContribution,
-			bIgnoreTranslucentMaterials);
 	}
 
 #if USE_IRRADIANCE_CACHING
@@ -519,7 +541,7 @@ void PathTracingKernel(
 		
 		if(bShouldAccumulateToCacheRecord)
 		{
-			float3 IrradianceToAccumulate = Irradiance * RadianceProbe_Pdf; // Revert the radiance probe pdf as we're caching the first bounce outgoing radiance
+			float3 IrradianceToAccumulate = Radiance * RadianceProbe_Pdf; // Revert the radiance probe pdf as we're caching the first bounce outgoing radiance
 			ATOMIC_ADD_FLOAT(IrradianceCachingParameters.IrradianceCacheRecords[NearestCacheEntryID].Irradiance.r, IrradianceToAccumulate.r);
 			ATOMIC_ADD_FLOAT(IrradianceCachingParameters.IrradianceCacheRecords[NearestCacheEntryID].Irradiance.g, IrradianceToAccumulate.g);
 			ATOMIC_ADD_FLOAT(IrradianceCachingParameters.IrradianceCacheRecords[NearestCacheEntryID].Irradiance.b, IrradianceToAccumulate.b);
@@ -528,7 +550,7 @@ void PathTracingKernel(
 	}
 #endif
 	
-	RadianceValue = Irradiance;
+	RadianceValue = Radiance;
 }
 
 [shader("raygeneration")]
diff --git a/Engine/Plugins/Experimental/GPULightmass/Source/GPULightmass/Private/LightmapDenoising.cpp b/Engine/Plugins/Experimental/GPULightmass/Source/GPULightmass/Private/LightmapDenoising.cpp
--- a/Engine/Plugins/Experimental/GPULightmass/Source/GPULightmass/Private/LightmapDenoising.cpp
+++ b/Engine/Plugins/Experimental/GPULightmass/Source/GPULightmass/Private/LightmapDenoising.cpp
@@ -124,30 +124,15 @@ void DenoiseLightSampleData(FIntPoint Size, TArray<FLightSampleData>& LightSampl
 	}
 
 	{
-		TArray<FVector> SHPositiveOutput;
-		TArray<FVector> SHNegativeOutput;
-		SHPositiveOutput.AddZeroed(Size.X * Size.Y);
-		SHNegativeOutput.AddZeroed(Size.X * Size.Y);
+		FLinearColor MinValue(MAX_flt, MAX_flt, MAX_flt);
 
 		for (int32 Y = 0; Y < Size.Y; Y++)
 		{
 			for (int32 X = 0; X < Size.X; X++)
 			{
-				FilterSet.InputBuffer[Y * FilterSet.Size.X + X][0] = FMath::Max(0.0f, LightSampleData[Y * Size.X + X].Coefficients[1][0]);
-				FilterSet.InputBuffer[Y * FilterSet.Size.X + X][1] = FMath::Max(0.0f, LightSampleData[Y * Size.X + X].Coefficients[1][1]);
-				FilterSet.InputBuffer[Y * FilterSet.Size.X + X][2] = FMath::Max(0.0f, LightSampleData[Y * Size.X + X].Coefficients[1][2]);
-			}
-		}
-
-		FilterSet.Execute();
-
-		for (int32 Y = 0; Y < Size.Y; Y++)
-		{
-			for (int32 X = 0; X < Size.X; X++)
-			{
-				 SHPositiveOutput[Y * Size.X + X][0] = FilterSet.OutputBuffer[Y * FilterSet.Size.X + X][0];
-				 SHPositiveOutput[Y * Size.X + X][1] = FilterSet.OutputBuffer[Y * FilterSet.Size.X + X][1];
-				 SHPositiveOutput[Y * Size.X + X][2] = FilterSet.OutputBuffer[Y * FilterSet.Size.X + X][2];
+				MinValue.R = FMath::Min(MinValue.R, LightSampleData[Y * Size.X + X].Coefficients[1][0]);
+				MinValue.G = FMath::Min(MinValue.G, LightSampleData[Y * Size.X + X].Coefficients[1][1]);
+				MinValue.B = FMath::Min(MinValue.B, LightSampleData[Y * Size.X + X].Coefficients[1][2]);
 			}
 		}
 
@@ -155,9 +140,9 @@ void DenoiseLightSampleData(FIntPoint Size, TArray<FLightSampleData>& LightSampl
 		{
 			for (int32 X = 0; X < Size.X; X++)
 			{
-				FilterSet.InputBuffer[Y * FilterSet.Size.X + X][0] = FMath::Max(0.0f, -LightSampleData[Y * Size.X + X].Coefficients[1][0]);
-				FilterSet.InputBuffer[Y * FilterSet.Size.X + X][1] = FMath::Max(0.0f, -LightSampleData[Y * Size.X + X].Coefficients[1][1]);
-				FilterSet.InputBuffer[Y * FilterSet.Size.X + X][2] = FMath::Max(0.0f, -LightSampleData[Y * Size.X + X].Coefficients[1][2]);
+				FilterSet.InputBuffer[Y * FilterSet.Size.X + X][0] = LightSampleData[Y * Size.X + X].Coefficients[1][0] - MinValue.R;
+				FilterSet.InputBuffer[Y * FilterSet.Size.X + X][1] = LightSampleData[Y * Size.X + X].Coefficients[1][1] - MinValue.G;
+				FilterSet.InputBuffer[Y * FilterSet.Size.X + X][2] = LightSampleData[Y * Size.X + X].Coefficients[1][2] - MinValue.B;
 			}
 		}
 
@@ -167,19 +152,9 @@ void DenoiseLightSampleData(FIntPoint Size, TArray<FLightSampleData>& LightSampl
 		{
 			for (int32 X = 0; X < Size.X; X++)
 			{
-				SHNegativeOutput[Y * Size.X + X][0] = FilterSet.OutputBuffer[Y * FilterSet.Size.X + X][0];
-				SHNegativeOutput[Y * Size.X + X][1] = FilterSet.OutputBuffer[Y * FilterSet.Size.X + X][1];
-				SHNegativeOutput[Y * Size.X + X][2] = FilterSet.OutputBuffer[Y * FilterSet.Size.X + X][2];
-			}
-		}
-
-		for (int32 Y = 0; Y < Size.Y; Y++)
-		{
-			for (int32 X = 0; X < Size.X; X++)
-			{
-				LightSampleData[Y * Size.X + X].Coefficients[1][0] = SHPositiveOutput[Y * Size.X + X][0] - SHNegativeOutput[Y * Size.X + X][0];
-				LightSampleData[Y * Size.X + X].Coefficients[1][1] = SHPositiveOutput[Y * Size.X + X][1] - SHNegativeOutput[Y * Size.X + X][1];
-				LightSampleData[Y * Size.X + X].Coefficients[1][2] = SHPositiveOutput[Y * Size.X + X][2] - SHNegativeOutput[Y * Size.X + X][2];
+				LightSampleData[Y * Size.X + X].Coefficients[1][0] = FilterSet.OutputBuffer[Y * FilterSet.Size.X + X][0] + MinValue.R;
+				LightSampleData[Y * Size.X + X].Coefficients[1][1] = FilterSet.OutputBuffer[Y * FilterSet.Size.X + X][1] + MinValue.G;
+				LightSampleData[Y * Size.X + X].Coefficients[1][2] = FilterSet.OutputBuffer[Y * FilterSet.Size.X + X][2] + MinValue.B;
 			}
 		}
 	}
@@ -310,30 +285,15 @@ void DenoiseRawData(
 	}
 
 	{
-		TArray<FVector> SHPositiveOutput;
-		TArray<FVector> SHNegativeOutput;
-		SHPositiveOutput.AddZeroed(Size.X * Size.Y);
-		SHNegativeOutput.AddZeroed(Size.X * Size.Y);
+		FLinearColor MinValue(MAX_flt, MAX_flt, MAX_flt);
 
 		for (int32 Y = 0; Y < Size.Y; Y++)
 		{
 			for (int32 X = 0; X < Size.X; X++)
 			{
-				FilterSet.InputBuffer[Y * FilterSet.Size.X + X][0] = FMath::Max(0.0f, LuminanceSH[Y * Size.X + X].R);
-				FilterSet.InputBuffer[Y * FilterSet.Size.X + X][1] = FMath::Max(0.0f, LuminanceSH[Y * Size.X + X].G);
-				FilterSet.InputBuffer[Y * FilterSet.Size.X + X][2] = FMath::Max(0.0f, LuminanceSH[Y * Size.X + X].B);
-			}
-		}
-
-		FilterSet.Execute();
-
-		for (int32 Y = 0; Y < Size.Y; Y++)
-		{
-			for (int32 X = 0; X < Size.X; X++)
-			{
-				SHPositiveOutput[Y * Size.X + X][0] = FilterSet.OutputBuffer[Y * FilterSet.Size.X + X][0];
-				SHPositiveOutput[Y * Size.X + X][1] = FilterSet.OutputBuffer[Y * FilterSet.Size.X + X][1];
-				SHPositiveOutput[Y * Size.X + X][2] = FilterSet.OutputBuffer[Y * FilterSet.Size.X + X][2];
+				MinValue.R = FMath::Min(MinValue.R, LuminanceSH[Y * Size.X + X].R);
+				MinValue.G = FMath::Min(MinValue.G, LuminanceSH[Y * Size.X + X].G);
+				MinValue.B = FMath::Min(MinValue.B, LuminanceSH[Y * Size.X + X].B);
 			}
 		}
 
@@ -341,9 +301,9 @@ void DenoiseRawData(
 		{
 			for (int32 X = 0; X < Size.X; X++)
 			{
-				FilterSet.InputBuffer[Y * FilterSet.Size.X + X][0] = FMath::Max(0.0f, -LuminanceSH[Y * Size.X + X].R);
-				FilterSet.InputBuffer[Y * FilterSet.Size.X + X][1] = FMath::Max(0.0f, -LuminanceSH[Y * Size.X + X].G);
-				FilterSet.InputBuffer[Y * FilterSet.Size.X + X][2] = FMath::Max(0.0f, -LuminanceSH[Y * Size.X + X].B);
+				FilterSet.InputBuffer[Y * FilterSet.Size.X + X][0] = LuminanceSH[Y * Size.X + X].R - MinValue.R;
+				FilterSet.InputBuffer[Y * FilterSet.Size.X + X][1] = LuminanceSH[Y * Size.X + X].G - MinValue.G;
+				FilterSet.InputBuffer[Y * FilterSet.Size.X + X][2] = LuminanceSH[Y * Size.X + X].B - MinValue.B;
 			}
 		}
 
@@ -353,19 +313,9 @@ void DenoiseRawData(
 		{
 			for (int32 X = 0; X < Size.X; X++)
 			{
-				SHNegativeOutput[Y * Size.X + X][0] = FilterSet.OutputBuffer[Y * FilterSet.Size.X + X][0];
-				SHNegativeOutput[Y * Size.X + X][1] = FilterSet.OutputBuffer[Y * FilterSet.Size.X + X][1];
-				SHNegativeOutput[Y * Size.X + X][2] = FilterSet.OutputBuffer[Y * FilterSet.Size.X + X][2];
-			}
-		}
-
-		for (int32 Y = 0; Y < Size.Y; Y++)
-		{
-			for (int32 X = 0; X < Size.X; X++)
-			{
-				LuminanceSH[Y * Size.X + X].R = SHPositiveOutput[Y * Size.X + X][0] - SHNegativeOutput[Y * Size.X + X][0];
-				LuminanceSH[Y * Size.X + X].G = SHPositiveOutput[Y * Size.X + X][1] - SHNegativeOutput[Y * Size.X + X][1];
-				LuminanceSH[Y * Size.X + X].B = SHPositiveOutput[Y * Size.X + X][2] - SHNegativeOutput[Y * Size.X + X][2];
+				LuminanceSH[Y * Size.X + X].R = FilterSet.OutputBuffer[Y * FilterSet.Size.X + X][0] + MinValue.R;
+				LuminanceSH[Y * Size.X + X].G = FilterSet.OutputBuffer[Y * FilterSet.Size.X + X][1] + MinValue.G;
+				LuminanceSH[Y * Size.X + X].B = FilterSet.OutputBuffer[Y * FilterSet.Size.X + X][2] + MinValue.B;
 			}
 		}
 	}
diff --git a/Engine/Plugins/Experimental/GPULightmass/Source/GPULightmass/Private/LightmapRenderer.cpp b/Engine/Plugins/Experimental/GPULightmass/Source/GPULightmass/Private/LightmapRenderer.cpp
--- a/Engine/Plugins/Experimental/GPULightmass/Source/GPULightmass/Private/LightmapRenderer.cpp
+++ b/Engine/Plugins/Experimental/GPULightmass/Source/GPULightmass/Private/LightmapRenderer.cpp
@@ -207,7 +207,7 @@ FPathTracingLightData SetupPathTracingLightParameters(const GPULightmass::FLight
 		{
 			LightParameters.Type[LightParameters.Count] = 1;
 			LightParameters.Position[LightParameters.Count] = Light.Position;
-			LightParameters.Color[LightParameters.Count] = FVector(Light.Color) / (4.0 * PI);
+			LightParameters.Color[LightParameters.Count] = FVector(Light.Color);
 			LightParameters.Dimensions[LightParameters.Count] = FVector(0.0, 0.0, Light.SourceRadius);
 			LightParameters.Attenuation[LightParameters.Count] = Light.AttenuationRadius;
 			LightParameters.Mobility[LightParameters.Count] = Light.bStationary ? 1 : 0;
@@ -226,7 +226,7 @@ FPathTracingLightData SetupPathTracingLightParameters(const GPULightmass::FLight
 			LightParameters.Type[LightParameters.Count] = 4;
 			LightParameters.Position[LightParameters.Count] = Light.Position;
 			LightParameters.Normal[LightParameters.Count] = Light.Direction;
-			LightParameters.Color[LightParameters.Count] = 4.0 * PI * FVector(Light.Color);
+			LightParameters.Color[LightParameters.Count] = FVector(Light.Color);
 			LightParameters.Dimensions[LightParameters.Count] = FVector(Light.SpotAngles, Light.SourceRadius);
 			LightParameters.Attenuation[LightParameters.Count] = Light.AttenuationRadius;
 			LightParameters.Mobility[LightParameters.Count] = Light.bStationary ? 1 : 0;
@@ -379,7 +379,7 @@ FLightmapRenderer::FLightmapRenderer(FSceneRenderState* InScene)
 	bDenoiseDuringInteractiveBake = Scene->Settings->DenoisingOptions == EGPULightmassDenoisingOptions::DuringInteractivePreview;
 	bOnlyBakeWhatYouSee = Scene->Settings->Mode == EGPULightmassMode::BakeWhatYouSee;
 	DenoisingThreadPool = FQueuedThreadPool::Allocate();
-	DenoisingThreadPool->Create(1);
+	DenoisingThreadPool->Create(1, 64 * 1024 * 1024);
 
 	if (bOnlyBakeWhatYouSee)
 	{
diff --git a/Engine/Shaders/Private/PathTracing/Material/PathTracingRadianceProbe.ush b/Engine/Shaders/Private/PathTracing/Material/PathTracingRadianceProbe.ush
--- a/Engine/Shaders/Private/PathTracing/Material/PathTracingRadianceProbe.ush
+++ b/Engine/Shaders/Private/PathTracing/Material/PathTracingRadianceProbe.ush
@@ -127,8 +127,9 @@ void RadianceProbe_SampleMaterial(
 	OutDirection = TangentToWorld(SampledValue.xyz, N_World);
 
 	PdfAdjust = max(PdfAdjust, 0.00001f);
-	OutWeight = 1.0f / PdfAdjust;
 	OutPdf = SampledValue.w * PdfAdjust;
+	// As the name suggests, radiance probe has a combined throughput (Pdf * Weight) == 1
+	OutWeight = (OutPdf > 0.0f) ? (1.0f / OutPdf) : 0;
 }
 
 float RadianceProbe_Pdf(float3 OutgoingDirection, FMaterialClosestHitPayload Payload)
@@ -268,7 +269,8 @@ void RadianceProbe_EvalMaterial(
 	PdfAdjust = max(PdfAdjust, 0.00001f);
 
 	OutPdf = NoL / PI * PdfAdjust;
-	OutWeight = 1.0f / PdfAdjust;
+	// As the name suggests, radiance probe has a combined throughput (Pdf * Weight) == 1
+	OutWeight = (OutPdf > 0.0f) ? (1.0f / OutPdf) : 0;
 }
 
 #else
