Fix GPULightmass UI not showing up as LevelEditor UI structure has been changed in UE5

diff --git a/Engine/Plugins/Experimental/GPULightmass/Source/GPULightmassEditor/GPULightmassEditor.Build.cs b/Engine/Plugins/Experimental/GPULightmass/Source/GPULightmassEditor/GPULightmassEditor.Build.cs
--- a/Engine/Plugins/Experimental/GPULightmass/Source/GPULightmassEditor/GPULightmassEditor.Build.cs
+++ b/Engine/Plugins/Experimental/GPULightmass/Source/GPULightmassEditor/GPULightmassEditor.Build.cs
@@ -30,6 +30,7 @@ namespace UnrealBuildTool.Rules
 					"GPULightmass",
 					"WorkspaceMenuStructure",
 					"EditorStyle",
+					"ToolMenus",
 				}
 				);
 		}
diff --git a/Engine/Plugins/Experimental/GPULightmass/Source/GPULightmassEditor/Private/GPULightmassEditorModule.cpp b/Engine/Plugins/Experimental/GPULightmass/Source/GPULightmassEditor/Private/GPULightmassEditorModule.cpp
--- a/Engine/Plugins/Experimental/GPULightmass/Source/GPULightmassEditor/Private/GPULightmassEditorModule.cpp
+++ b/Engine/Plugins/Experimental/GPULightmass/Source/GPULightmassEditor/Private/GPULightmassEditorModule.cpp
@@ -21,6 +21,7 @@
 #include "Framework/MultiBox/MultiBoxBuilder.h"
 #include "Widgets/Docking/SDockTab.h"
 #include "Widgets/Input/SButton.h"
+#include "ToolMenus.h"
 
 extern ENGINE_API void ToggleLightmapPreview_GameThread(UWorld* InWorld);
 
@@ -38,6 +39,19 @@ void FGPULightmassEditorModule::StartupModule()
 	auto BuildMenuExtender = FLevelEditorModule::FLevelEditorMenuExtender::CreateRaw(this, &FGPULightmassEditorModule::OnExtendLevelEditorBuildMenu);
 	LevelEditorModule.GetAllLevelEditorToolbarBuildMenuExtenders().Add(BuildMenuExtender);
 
+	UToolMenu* Menu = UToolMenus::Get()->ExtendMenu("LevelEditor.MainMenu.Build");
+	FToolMenuSection& Section = Menu->FindOrAddSection("LevelEditorLighting");
+
+	FUIAction ActionOpenGPULightmassSettingsTab(FExecuteAction::CreateLambda([]() {
+		FLevelEditorModule& LevelEditorModule = FModuleManager::GetModuleChecked<FLevelEditorModule>(TEXT("LevelEditor"));
+		TSharedPtr<FTabManager> LevelEditorTabManager = LevelEditorModule.GetLevelEditorTabManager();
+		LevelEditorTabManager->TryInvokeTab(GPULightmassSettingsTabName);
+	}), FCanExecuteAction());
+
+	FToolMenuEntry& Entry = Section.AddMenuEntry(NAME_None, LOCTEXT("GPULightmassSettingsTitle", "GPU Lightmass"),
+		LOCTEXT("OpensGPULightmassSettings", "Opens GPU Lightmass settings tab."), FSlateIcon(FEditorStyle::GetStyleSetName(), "Level.LightingScenarioIcon16x"), ActionOpenGPULightmassSettingsTab,
+		EUserInterfaceActionType::Button);
+
 	FGPULightmassModule& GPULightmassModule = FModuleManager::LoadModuleChecked<FGPULightmassModule>(TEXT("GPULightmass"));
 	GPULightmassModule.OnStaticLightingSystemsChanged.AddLambda([&SettingsView = SettingsView, &StartStopButtonText = StartStopButtonText]()
 	{ 
@@ -48,10 +62,13 @@ void FGPULightmassEditorModule::StartupModule()
 
 		if (StartStopButtonText.IsValid())
 		{
-			UWorld* World = GEditor->GetEditorWorldContext().World();
-			if (World)
+			if (GEditor)
 			{
-				StartStopButtonText->SetText(World->GetSubsystem<UGPULightmassSubsystem>()->IsRunning() ? LOCTEXT("GPULightmassSettingsStop", "Stop") : LOCTEXT("GPULightmassSettingsStart", "Build Lighting"));
+				UWorld* World = GEditor->GetEditorWorldContext().World();
+				if (World)
+				{
+					StartStopButtonText->SetText(World->GetSubsystem<UGPULightmassSubsystem>()->IsRunning() ? LOCTEXT("GPULightmassSettingsStop", "Stop") : LOCTEXT("GPULightmassSettingsStart", "Build Lighting"));
+				}
 			}
 		}
 	});
