Converted virtual texture system and GPU lightmass to RDG.

diff --git a/Engine/Plugins/Experimental/GPULightmass/Source/GPULightmass/Private/LightmapGBuffer.cpp b/Engine/Plugins/Experimental/GPULightmass/Source/GPULightmass/Private/LightmapGBuffer.cpp
--- a/Engine/Plugins/Experimental/GPULightmass/Source/GPULightmass/Private/LightmapGBuffer.cpp
+++ b/Engine/Plugins/Experimental/GPULightmass/Source/GPULightmass/Private/LightmapGBuffer.cpp
@@ -2,7 +2,7 @@
 
 #include "LightmapGBuffer.h"
 
-IMPLEMENT_GLOBAL_SHADER_PARAMETER_STRUCT(FLightmapGBufferParams, "LightmapGBufferParams");
+IMPLEMENT_STATIC_UNIFORM_BUFFER_STRUCT(FLightmapGBufferParams, "LightmapGBufferParams", SceneTextures);
 
 IMPLEMENT_MATERIAL_SHADER_TYPE(, FLightmapGBufferVS, TEXT("/Plugin/GPULightmass/Private/LightmapGBuffer.usf"), TEXT("LightmapGBufferVS"), SF_Vertex);
 IMPLEMENT_MATERIAL_SHADER_TYPE(, FLightmapGBufferPS, TEXT("/Plugin/GPULightmass/Private/LightmapGBuffer.usf"), TEXT("LightmapGBufferPS"), SF_Pixel);
diff --git a/Engine/Plugins/Experimental/GPULightmass/Source/GPULightmass/Private/LightmapGBuffer.h b/Engine/Plugins/Experimental/GPULightmass/Source/GPULightmass/Private/LightmapGBuffer.h
--- a/Engine/Plugins/Experimental/GPULightmass/Source/GPULightmass/Private/LightmapGBuffer.h
+++ b/Engine/Plugins/Experimental/GPULightmass/Source/GPULightmass/Private/LightmapGBuffer.h
@@ -10,13 +10,11 @@
 #include "LightMapRendering.h"
 
 BEGIN_GLOBAL_SHADER_PARAMETER_STRUCT(FLightmapGBufferParams, )
-	SHADER_PARAMETER_UAV(RWTexture2D<float4>, ScratchTilePoolLayer0)
-	SHADER_PARAMETER_UAV(RWTexture2D<float4>, ScratchTilePoolLayer1)
-	SHADER_PARAMETER_UAV(RWTexture2D<float4>, ScratchTilePoolLayer2)
+	SHADER_PARAMETER_RDG_TEXTURE_UAV(RWTexture2D<float4>, ScratchTilePoolLayer0)
+	SHADER_PARAMETER_RDG_TEXTURE_UAV(RWTexture2D<float4>, ScratchTilePoolLayer1)
+	SHADER_PARAMETER_RDG_TEXTURE_UAV(RWTexture2D<float4>, ScratchTilePoolLayer2)
 END_GLOBAL_SHADER_PARAMETER_STRUCT()
 
-typedef TUniformBufferRef<FLightmapGBufferParams> FLightmapGBufferUniformBufferRef;
-
 struct FLightmapElementData : public FMeshMaterialShaderElementData
 {
 	const FLightCacheInterface* LCI;
@@ -41,7 +39,6 @@ protected:
 	FLightmapGBufferVS(const ShaderMetaType::CompiledShaderInitializerType& Initializer)
 		: FMeshMaterialShader(Initializer)
 	{
-		PassUniformBuffer.Bind(Initializer.ParameterMap, FLightmapGBufferParams::StaticStructMetadata.GetShaderVariableName());
 		PrecomputedLightingBufferParameter.Bind(Initializer.ParameterMap, TEXT("PrecomputedLightingBuffer"));
 		VirtualTexturePhysicalTileCoordinateScaleAndBias.Bind(Initializer.ParameterMap, TEXT("VirtualTexturePhysicalTileCoordinateScaleAndBias"));
 		RenderPassIndex.Bind(Initializer.ParameterMap, TEXT("RenderPassIndex"));
@@ -142,7 +139,6 @@ public:
 	FLightmapGBufferPS(const ShaderMetaType::CompiledShaderInitializerType& Initializer)
 		: FMeshMaterialShader(Initializer)
 	{
-		PassUniformBuffer.Bind(Initializer.ParameterMap, FLightmapGBufferParams::StaticStructMetadata.GetShaderVariableName());
 		ScratchTilePoolOffset.Bind(Initializer.ParameterMap, TEXT("ScratchTilePoolOffset"));
 	}
 
@@ -167,14 +163,13 @@ public:
 	FLightmapGBufferMeshProcessor(
 		const FScene* InScene, 
 		const FSceneView* InView,
-		FMeshPassDrawListContext* InDrawListContext, 
-		FRHIUniformBuffer* InPassUniformBuffer,
+		FMeshPassDrawListContext* InDrawListContext,
 		FVector4 VirtualTexturePhysicalTileCoordinateScaleAndBias,
 		int32 RenderPassIndex,
 		FIntPoint ScratchTilePoolOffset
 	)
 		: FMeshPassProcessor(InScene, InView->GetFeatureLevel(), InView, InDrawListContext)
-		, DrawRenderState(*InView, InPassUniformBuffer)
+		, DrawRenderState(*InView)
 		, VirtualTexturePhysicalTileCoordinateScaleAndBias(VirtualTexturePhysicalTileCoordinateScaleAndBias)
 		, RenderPassIndex(RenderPassIndex)
 		, ScratchTilePoolOffset(ScratchTilePoolOffset)
diff --git a/Engine/Plugins/Experimental/GPULightmass/Source/GPULightmass/Private/LightmapRenderer.cpp b/Engine/Plugins/Experimental/GPULightmass/Source/GPULightmass/Private/LightmapRenderer.cpp
--- a/Engine/Plugins/Experimental/GPULightmass/Source/GPULightmass/Private/LightmapRenderer.cpp
+++ b/Engine/Plugins/Experimental/GPULightmass/Source/GPULightmass/Private/LightmapRenderer.cpp
@@ -133,7 +133,7 @@ class FMultiTileClearCS : public FGlobalShader
 		SHADER_PARAMETER(int32, NumTiles)
 		SHADER_PARAMETER(int32, TileSize)
 		SHADER_PARAMETER_SRV(Buffer<int2>, TilePositions)
-		SHADER_PARAMETER_UAV(RWTexture2D<float4>, TilePool)
+		SHADER_PARAMETER_RDG_TEXTURE_UAV(RWTexture2D<float4>, TilePool)
 	END_SHADER_PARAMETER_STRUCT()
 };
 
@@ -664,7 +664,7 @@ void FSceneRenderState::SetupRayTracingScene()
 	ViewInitOptions.BackgroundColor = FLinearColor::Black;
 	ViewInitOptions.OverlayColor = FLinearColor::White;
 
-	ReferenceView = MakeUnique<FViewInfo>(ViewInitOptions);
+	ReferenceView = MakeShared<FViewInfo>(ViewInitOptions);
 	FViewInfo& View = *ReferenceView;
 	View.ViewRect = View.UnscaledViewRect;
 
@@ -915,7 +915,12 @@ void FSceneRenderState::CalculateDistributionPrefixSumForAllLightmaps()
 	}
 }
 
-void FLightmapRenderer::Finalize(FRHICommandListImmediate& RHICmdList)
+BEGIN_SHADER_PARAMETER_STRUCT(FLightmapGBufferPassParameters, )
+	SHADER_PARAMETER_RDG_UNIFORM_BUFFER(FLightmapGBufferParams, PassUniformBuffer)
+	RENDER_TARGET_BINDING_SLOTS()
+END_SHADER_PARAMETER_STRUCT()
+
+void FLightmapRenderer::Finalize(FRDGBuilder& GraphBuilder)
 {
 	TRACE_CPUPROFILER_EVENT_SCOPE(FLightmapRenderer::Finalize);
 
@@ -924,7 +929,10 @@ void FLightmapRenderer::Finalize(FRHICommandListImmediate& RHICmdList)
 		return;
 	}
 
-	FMemMark Mark(FMemStack::Get());
+	const auto HoldReference = [](FRDGBuilder& GraphBuilder, const FShaderResourceViewRHIRef& View) -> const FShaderResourceViewRHIRef&
+	{
+		return *GraphBuilder.AllocObject<FShaderResourceViewRHIRef>(View);
+	};
 
 	// Upload & copy converged tiles directly
 	{
@@ -936,6 +944,7 @@ void FLightmapRenderer::Finalize(FRHICommandListImmediate& RHICmdList)
 
 		if (TileUploadRequests.Num() > 0)
 		{
+			FRHICommandListImmediate& RHICmdList = GraphBuilder.RHICmdList;
 			SCOPED_DRAW_EVENTF(RHICmdList, GPULightmassUploadConvergedTiles, TEXT("GPULightmass UploadConvergedTiles %d tiles"), TileUploadRequests.Num());
 
 			int32 NewSize = FMath::CeilToInt(FMath::Sqrt(TileUploadRequests.Num()));
@@ -1033,8 +1042,6 @@ void FLightmapRenderer::Finalize(FRHICommandListImmediate& RHICmdList)
 			DispatchResolution.X = GPreviewLightmapPhysicalTileSize * GPUBatchedTileRequests.BatchedTilesDesc.Num();
 			DispatchResolution.Y = GPreviewLightmapPhysicalTileSize;
 
-			FRDGBuilder GraphBuilder(RHICmdList);
-
 			FRDGTextureRef StagingHQLayer0 = GraphBuilder.RegisterExternalTexture(UploadTilePoolGPU->PooledRenderTargets[0], TEXT("StagingHQLayer0"));
 			FRDGTextureRef StagingHQLayer1 = GraphBuilder.RegisterExternalTexture(UploadTilePoolGPU->PooledRenderTargets[1], TEXT("StagingHQLayer1"));
 			FRDGTextureRef StagingShadowMask = GraphBuilder.RegisterExternalTexture(UploadTilePoolGPU->PooledRenderTargets[2], TEXT("StagingShadowMask"));
@@ -1077,8 +1084,8 @@ void FLightmapRenderer::Finalize(FRHICommandListImmediate& RHICmdList)
 					PassParameters->NumBatchedTiles = GPUBatchedTileRequests.BatchedTilesDesc.Num();
 					PassParameters->SrcTexture = GraphBuilder.CreateUAV(StagingHQLayer0);
 					PassParameters->DstTexture = GraphBuilder.CreateUAV(RenderTargetTileAtlas0);
-					PassParameters->SrcTilePositions = SrcTilePositionsSRV;
-					PassParameters->DstTilePositions = DstTilePositionsSRV;
+					PassParameters->SrcTilePositions = HoldReference(GraphBuilder, SrcTilePositionsSRV);
+					PassParameters->DstTilePositions = HoldReference(GraphBuilder, DstTilePositionsSRV);
 
 					TShaderMapRef<FUploadConvergedLightmapTilesCS> ComputeShader(GetGlobalShaderMap(GMaxRHIFeatureLevel));
 					FComputeShaderUtils::AddPass(
@@ -1123,8 +1130,8 @@ void FLightmapRenderer::Finalize(FRHICommandListImmediate& RHICmdList)
 					PassParameters->NumBatchedTiles = GPUBatchedTileRequests.BatchedTilesDesc.Num();
 					PassParameters->SrcTexture = GraphBuilder.CreateUAV(StagingHQLayer1);
 					PassParameters->DstTexture = GraphBuilder.CreateUAV(RenderTargetTileAtlas1);
-					PassParameters->SrcTilePositions = SrcTilePositionsSRV;
-					PassParameters->DstTilePositions = DstTilePositionsSRV;
+					PassParameters->SrcTilePositions = HoldReference(GraphBuilder, SrcTilePositionsSRV);
+					PassParameters->DstTilePositions = HoldReference(GraphBuilder, DstTilePositionsSRV);
 
 					TShaderMapRef<FUploadConvergedLightmapTilesCS> ComputeShader(GetGlobalShaderMap(GMaxRHIFeatureLevel));
 					FComputeShaderUtils::AddPass(
@@ -1169,8 +1176,8 @@ void FLightmapRenderer::Finalize(FRHICommandListImmediate& RHICmdList)
 					PassParameters->NumBatchedTiles = GPUBatchedTileRequests.BatchedTilesDesc.Num();
 					PassParameters->SrcTexture = GraphBuilder.CreateUAV(StagingShadowMask);
 					PassParameters->DstTexture = GraphBuilder.CreateUAV(RenderTargetTileAtlas2);
-					PassParameters->SrcTilePositions = SrcTilePositionsSRV;
-					PassParameters->DstTilePositions = DstTilePositionsSRV;
+					PassParameters->SrcTilePositions = HoldReference(GraphBuilder, SrcTilePositionsSRV);
+					PassParameters->DstTilePositions = HoldReference(GraphBuilder, DstTilePositionsSRV);
 
 					TShaderMapRef<FUploadConvergedLightmapTilesCS> ComputeShader(GetGlobalShaderMap(GMaxRHIFeatureLevel));
 					FComputeShaderUtils::AddPass(
@@ -1181,8 +1188,6 @@ void FLightmapRenderer::Finalize(FRHICommandListImmediate& RHICmdList)
 						FComputeShaderUtils::GetGroupCount(DispatchResolution, FComputeShaderUtils::kGolden2DGroupSize));
 				}
 			}
-
-			GraphBuilder.Execute();
 		}
 
 		// Drop these converged requests, critical so that we won't perform readback repeatedly
@@ -1345,13 +1350,24 @@ void FLightmapRenderer::Finalize(FRHICommandListImmediate& RHICmdList)
 
 	Scene->SetupRayTracingScene();
 
-	FLightmapGBufferParams LightmapGBufferParameters{};
-	LightmapGBufferParameters.ScratchTilePoolLayer0 = ScratchTilePoolGPU->PooledRenderTargets[0]->GetRenderTargetItem().UAV;
-	LightmapGBufferParameters.ScratchTilePoolLayer1 = ScratchTilePoolGPU->PooledRenderTargets[1]->GetRenderTargetItem().UAV;
-	LightmapGBufferParameters.ScratchTilePoolLayer2 = ScratchTilePoolGPU->PooledRenderTargets[2]->GetRenderTargetItem().UAV;
-	FLightmapGBufferUniformBufferRef PassUniformBuffer = FLightmapGBufferUniformBufferRef::CreateUniformBufferImmediate(LightmapGBufferParameters, UniformBuffer_SingleFrame);
+	TStaticArray<FRDGTextureUAVRef, 3> ScratchTilePoolLayerUAVs;
 
-	SCOPED_GPU_MASK(RHICmdList, FRHIGPUMask::GPU0());
+	for (int32 Index = 0; Index < ScratchTilePoolLayerUAVs.Num(); ++Index)
+	{
+		ScratchTilePoolLayerUAVs[Index] = GraphBuilder.CreateUAV(GraphBuilder.RegisterExternalTexture(ScratchTilePoolGPU->PooledRenderTargets[Index]));
+	}
+
+	TRDGUniformBufferRef<FLightmapGBufferParams> PassUniformBuffer = nullptr;
+
+	{
+		auto* LightmapGBufferParameters = GraphBuilder.AllocParameters<FLightmapGBufferParams>();
+		LightmapGBufferParameters->ScratchTilePoolLayer0 = ScratchTilePoolLayerUAVs[0];
+		LightmapGBufferParameters->ScratchTilePoolLayer1 = ScratchTilePoolLayerUAVs[1];
+		LightmapGBufferParameters->ScratchTilePoolLayer2 = ScratchTilePoolLayerUAVs[2];
+		PassUniformBuffer = GraphBuilder.CreateUniformBuffer(LightmapGBufferParameters);
+	}
+
+	RDG_GPU_MASK_SCOPE(GraphBuilder, FRHIGPUMask::GPU0());
 
 	IPooledRenderTarget* OutputRenderTargets[3] = { nullptr, nullptr, nullptr };
 
@@ -1378,26 +1394,43 @@ void FLightmapRenderer::Finalize(FRHICommandListImmediate& RHICmdList)
 		// Clear working set pools
 		for (int PoolLayerIndex = 0; PoolLayerIndex < LightmapTilePoolGPU.PooledRenderTargets.Num(); PoolLayerIndex++)
 		{
-			SCOPED_GPU_MASK(RHICmdList, FRHIGPUMask::All());
-			FRHIRenderPassInfo RPInfo(LightmapTilePoolGPU.PooledRenderTargets[PoolLayerIndex]->GetRenderTargetItem().TargetableTexture, ERenderTargetActions::DontLoad_Store);
-			RHICmdList.BeginRenderPass(RPInfo, TEXT("ClearLightmapTilePoolGPU"));
+			RDG_GPU_MASK_SCOPE(GraphBuilder, FRHIGPUMask::All());
+
+			TArray<FVector4, SceneRenderingAllocator> ViewportsToClear;
+
 			for (auto& Tile : PendingTileRequests)
 			{
 				if (Tile.RenderState->RetrieveTileState(Tile.VirtualCoordinates).Revision != CurrentRevision)
 				{
-					RHICmdList.SetViewport(
+					ViewportsToClear.Emplace(
 						LightmapTilePoolGPU.GetPositionFromLinearAddress(Tile.TileAddressInWorkingSet).X * LightmapTilePoolGPU.LayerFormatAndTileSize[PoolLayerIndex].TileSize.X,
 						LightmapTilePoolGPU.GetPositionFromLinearAddress(Tile.TileAddressInWorkingSet).Y * LightmapTilePoolGPU.LayerFormatAndTileSize[PoolLayerIndex].TileSize.Y,
-						0.0f,
 						(LightmapTilePoolGPU.GetPositionFromLinearAddress(Tile.TileAddressInWorkingSet).X + 1) * LightmapTilePoolGPU.LayerFormatAndTileSize[PoolLayerIndex].TileSize.X,
-						(LightmapTilePoolGPU.GetPositionFromLinearAddress(Tile.TileAddressInWorkingSet).Y + 1) * LightmapTilePoolGPU.LayerFormatAndTileSize[PoolLayerIndex].TileSize.Y,
-						1.0f);
-					DrawClearQuad(RHICmdList, FLinearColor(0.0f, 0.0f, 0.0f, 0.0f));
+						(LightmapTilePoolGPU.GetPositionFromLinearAddress(Tile.TileAddressInWorkingSet).Y + 1) * LightmapTilePoolGPU.LayerFormatAndTileSize[PoolLayerIndex].TileSize.Y);
 				}
 			}
-			RHICmdList.EndRenderPass();
-		}
 
+			if (ViewportsToClear.Num())
+			{
+				FRDGTextureRef Texture = GraphBuilder.RegisterExternalTexture(LightmapTilePoolGPU.PooledRenderTargets[PoolLayerIndex]);
+
+				auto* PassParameters = GraphBuilder.AllocParameters<FRenderTargetParameters>();
+				PassParameters->RenderTargets[0] = FRenderTargetBinding(Texture, ERenderTargetLoadAction::ENoAction);
+
+				GraphBuilder.AddPass(
+					RDG_EVENT_NAME("ClearLightmapTilePoolGPU"),
+					PassParameters,
+					ERDGPassFlags::Raster,
+					[LocalViewportsToClear = MoveTemp(ViewportsToClear)](FRHICommandList& RHICmdList)
+				{
+					for (FVector4 Viewport : LocalViewportsToClear)
+					{
+						RHICmdList.SetViewport(Viewport.X, Viewport.Y, 0.0f, Viewport.Z, Viewport.W, 1.0f);
+						DrawClearQuad(RHICmdList, FLinearColor(0.0f, 0.0f, 0.0f, 0.0f));
+					}
+				});
+			}
+		}
 
 		for (auto& Tile : PendingTileRequests)
 		{
@@ -1457,20 +1490,20 @@ void FLightmapRenderer::Finalize(FRHICommandListImmediate& RHICmdList)
 	int32 NumSamplesPerFrame = (bInsideBackgroundTick && bLastFewFramesIdle) ? Scene->Settings->TilePassesInFullSpeedMode : Scene->Settings->TilePassesInSlowMode;
 
 	{
-		TArray<FLightmapTileRequest> PendingGITileRequests = PendingTileRequests.FilterByPredicate([NumGISamples = Scene->Settings->GISamples](const FLightmapTileRequest& Tile) { return !Tile.RenderState->IsTileGIConverged(Tile.VirtualCoordinates, NumGISamples); });
+		auto& PendingGITileRequests = *GraphBuilder.AllocObject<TArray<FLightmapTileRequest>>(
+			PendingTileRequests.FilterByPredicate([NumGISamples = Scene->Settings->GISamples](const FLightmapTileRequest& Tile)
+			{
+				return !Tile.RenderState->IsTileGIConverged(Tile.VirtualCoordinates, NumGISamples);
+			}));
 
 		// Render GI
 		for (int32 SampleIndex = 0; SampleIndex < NumSamplesPerFrame; SampleIndex++)
 		{
-			FMemMark PerSampleMark(FMemStack::Get());
-
 			{
 				if (PendingGITileRequests.Num() > 0)
 				{
 					for (int ScratchLayerIndex = 0; ScratchLayerIndex < 3; ScratchLayerIndex++)
 					{
-						FRDGBuilder GraphBuilder(RHICmdList);
-
 						TResourceArray<FIntPoint> TilePositionsToClear;
 						for (auto& Tile : PendingGITileRequests)
 						{
@@ -1483,8 +1516,8 @@ void FLightmapRenderer::Finalize(FRHICommandListImmediate& RHICmdList)
 						FMultiTileClearCS::FParameters* Parameters = GraphBuilder.AllocParameters<FMultiTileClearCS::FParameters>();
 						Parameters->NumTiles = TilePositionsToClear.Num();
 						Parameters->TileSize = GPreviewLightmapPhysicalTileSize;
-						Parameters->TilePositions = TilePositionsBuffer.SRV;
-						Parameters->TilePool = ScratchTilePoolGPU->PooledRenderTargets[ScratchLayerIndex]->GetRenderTargetItem().UAV;
+						Parameters->TilePositions = HoldReference(GraphBuilder, TilePositionsBuffer.SRV);
+						Parameters->TilePool = ScratchTilePoolLayerUAVs[ScratchLayerIndex];
 
 						TShaderMapRef<FMultiTileClearCS> ComputeShader(GetGlobalShaderMap(GMaxRHIFeatureLevel));
 						FComputeShaderUtils::AddPass(
@@ -1493,70 +1526,67 @@ void FLightmapRenderer::Finalize(FRHICommandListImmediate& RHICmdList)
 							ComputeShader,
 							Parameters,
 							FComputeShaderUtils::GetGroupCount(FIntPoint(GPreviewLightmapPhysicalTileSize * TilePositionsToClear.Num(), GPreviewLightmapPhysicalTileSize), FComputeShaderUtils::kGolden2DGroupSize));
-
-						GraphBuilder.Execute();
-					}
-
-					for (int ScratchLayerIndex = 0; ScratchLayerIndex < 3; ScratchLayerIndex++)
-					{
-						RHICmdList.Transition(FRHITransitionInfo(ScratchTilePoolGPU->PooledRenderTargets[ScratchLayerIndex]->GetRenderTargetItem().UAV, ERHIAccess::UAVGraphics, ERHIAccess::UAVGraphics));
 					}
 
 					{
 						for (uint32 GPUIndex = 0; GPUIndex < GNumExplicitGPUsForRendering; GPUIndex++)
 						{
-							SCOPED_GPU_MASK(RHICmdList, FRHIGPUMask::FromIndex(GPUIndex));
+							RDG_GPU_MASK_SCOPE(GraphBuilder, FRHIGPUMask::FromIndex(GPUIndex));
 
-							FRHIRenderPassInfo RPInfo(FRHIRenderPassInfo::NoRenderTargets);
-							RHICmdList.BeginRenderPass(RPInfo, TEXT("LightmapGBuffer"));
+							auto* PassParameters = GraphBuilder.AllocParameters<FLightmapGBufferPassParameters>();
+							PassParameters->PassUniformBuffer = PassUniformBuffer;
 
-							for (const FLightmapTileRequest& Tile : PendingGITileRequests)
-							{							
-								if (Tile.RenderState->IsTileGIConverged(Tile.VirtualCoordinates, Scene->Settings->GISamples)) continue;
-								uint32 AssignedGPUIndex = (Tile.RenderState->DistributionPrefixSum + Tile.RenderState->RetrieveTileStateIndex(Tile.VirtualCoordinates)) % GNumExplicitGPUsForRendering;
-								if (AssignedGPUIndex != GPUIndex) continue;
+							GraphBuilder.AddPass(
+								RDG_EVENT_NAME("LightmapGBuffer"),
+								PassParameters,
+								ERDGPassFlags::Raster,
+								[this, ReferenceView = Scene->ReferenceView, &PendingGITileRequests, GPUIndex](FRHICommandList& RHICmdList)
+							{
+								for (const FLightmapTileRequest& Tile : PendingGITileRequests)
+								{
+									if (Tile.RenderState->IsTileGIConverged(Tile.VirtualCoordinates, Scene->Settings->GISamples)) continue;
+									uint32 AssignedGPUIndex = (Tile.RenderState->DistributionPrefixSum + Tile.RenderState->RetrieveTileStateIndex(Tile.VirtualCoordinates)) % GNumExplicitGPUsForRendering;
+									if (AssignedGPUIndex != GPUIndex) continue;
 
-								RHICmdList.SetViewport(0, 0, 0.0f, GPreviewLightmapPhysicalTileSize, GPreviewLightmapPhysicalTileSize, 1.0f);
+									RHICmdList.SetViewport(0, 0, 0.0f, GPreviewLightmapPhysicalTileSize, GPreviewLightmapPhysicalTileSize, 1.0f);
 
-								float ScaleX = Tile.RenderState->GetPaddedSizeInTiles().X * GPreviewLightmapVirtualTileSize * 1.0f / (1 << Tile.VirtualCoordinates.MipLevel) / GPreviewLightmapPhysicalTileSize;
-								float ScaleY = Tile.RenderState->GetPaddedSizeInTiles().Y * GPreviewLightmapVirtualTileSize * 1.0f / (1 << Tile.VirtualCoordinates.MipLevel) / GPreviewLightmapPhysicalTileSize;
-								float BiasX = (1.0f * (-Tile.VirtualCoordinates.Position.X * GPreviewLightmapVirtualTileSize) - (-GPreviewLightmapTileBorderSize)) / GPreviewLightmapPhysicalTileSize;
-								float BiasY = (1.0f * (-Tile.VirtualCoordinates.Position.Y * GPreviewLightmapVirtualTileSize) - (-GPreviewLightmapTileBorderSize)) / GPreviewLightmapPhysicalTileSize;
-
-								FVector4 VirtualTexturePhysicalTileCoordinateScaleAndBias = FVector4(ScaleX, ScaleY, BiasX, BiasY);
+									float ScaleX = Tile.RenderState->GetPaddedSizeInTiles().X * GPreviewLightmapVirtualTileSize * 1.0f / (1 << Tile.VirtualCoordinates.MipLevel) / GPreviewLightmapPhysicalTileSize;
+									float ScaleY = Tile.RenderState->GetPaddedSizeInTiles().Y * GPreviewLightmapVirtualTileSize * 1.0f / (1 << Tile.VirtualCoordinates.MipLevel) / GPreviewLightmapPhysicalTileSize;
+									float BiasX = (1.0f * (-Tile.VirtualCoordinates.Position.X * GPreviewLightmapVirtualTileSize) - (-GPreviewLightmapTileBorderSize)) / GPreviewLightmapPhysicalTileSize;
+									float BiasY = (1.0f * (-Tile.VirtualCoordinates.Position.Y * GPreviewLightmapVirtualTileSize) - (-GPreviewLightmapTileBorderSize)) / GPreviewLightmapPhysicalTileSize;
 
-								TArray<FMeshBatch> MeshBatches = Tile.RenderState->GeometryInstanceRef.GetMeshBatchesForGBufferRendering(Tile.VirtualCoordinates);
+									FVector4 VirtualTexturePhysicalTileCoordinateScaleAndBias = FVector4(ScaleX, ScaleY, BiasX, BiasY);
 
-								for (auto& MeshBatch : MeshBatches)
-								{
-									FMeshBatchElement& Element = MeshBatch.Elements[0];
-
-									Element.DynamicPrimitiveShaderDataIndex = Tile.RenderState->GeometryInstanceRef.GetElementId();
-								}
-
-								DrawDynamicMeshPass(
-									*Scene->ReferenceView, RHICmdList,
-									[
-										View = Scene->ReferenceView.Get(), 
-										PassUniformBuffer, 
-										MeshBatches,
-										VirtualTexturePhysicalTileCoordinateScaleAndBias,
-										RenderPassIndex = Tile.RenderState->RetrieveTileState(Tile.VirtualCoordinates).RenderPassIndex,
-										ScratchTilePoolOffset = ScratchTilePoolGPU->GetPositionFromLinearAddress(Tile.TileAddressInScratch) * GPreviewLightmapPhysicalTileSize
-									](FDynamicPassMeshDrawListContext* DynamicMeshPassContext)
-								{
-									FLightmapGBufferMeshProcessor MeshProcessor(nullptr, View, DynamicMeshPassContext, PassUniformBuffer, VirtualTexturePhysicalTileCoordinateScaleAndBias, RenderPassIndex, ScratchTilePoolOffset);
+									TArray<FMeshBatch> MeshBatches = Tile.RenderState->GeometryInstanceRef.GetMeshBatchesForGBufferRendering(Tile.VirtualCoordinates);
 
 									for (auto& MeshBatch : MeshBatches)
 									{
-										MeshProcessor.AddMeshBatch(MeshBatch, ~0ull, nullptr);
+										FMeshBatchElement& Element = MeshBatch.Elements[0];
+
+										Element.DynamicPrimitiveShaderDataIndex = Tile.RenderState->GeometryInstanceRef.GetElementId();
 									}
-								});
 
-								GPrimitiveIdVertexBufferPool.DiscardAll();
-							}
+									DrawDynamicMeshPass(
+										*ReferenceView, RHICmdList,
+										[
+											View = ReferenceView.Get(),
+											MeshBatches,
+											VirtualTexturePhysicalTileCoordinateScaleAndBias,
+											RenderPassIndex = Tile.RenderState->RetrieveTileState(Tile.VirtualCoordinates).RenderPassIndex,
+											ScratchTilePoolOffset = ScratchTilePoolGPU->GetPositionFromLinearAddress(Tile.TileAddressInScratch) * GPreviewLightmapPhysicalTileSize
+										](FDynamicPassMeshDrawListContext* DynamicMeshPassContext)
+									{
+										FLightmapGBufferMeshProcessor MeshProcessor(nullptr, View, DynamicMeshPassContext, VirtualTexturePhysicalTileCoordinateScaleAndBias, RenderPassIndex, ScratchTilePoolOffset);
 
-							RHICmdList.EndRenderPass();
+										for (auto& MeshBatch : MeshBatches)
+										{
+											MeshProcessor.AddMeshBatch(MeshBatch, ~0ull, nullptr);
+										}
+									});
+
+									GPrimitiveIdVertexBufferPool.DiscardAll();
+								}
+							});
 						}
 					}
 
@@ -1604,15 +1634,13 @@ void FLightmapRenderer::Finalize(FRHICommandListImmediate& RHICmdList)
 								CreateInfo.ResourceArray = &GPUBatchedTileRequests.BatchedTilesDesc;
 
 								GPUBatchedTileRequests.BatchedTilesBuffer = RHICreateStructuredBuffer(sizeof(FGPUTileDescription), GPUBatchedTileRequests.BatchedTilesDesc.GetResourceDataSize(), BUF_Static | BUF_ShaderResource, CreateInfo);
-								GPUBatchedTileRequests.BatchedTilesSRV = RHICreateShaderResourceView(GPUBatchedTileRequests.BatchedTilesBuffer);
+								GPUBatchedTileRequests.BatchedTilesSRV = HoldReference(GraphBuilder, RHICreateShaderResourceView(GPUBatchedTileRequests.BatchedTilesBuffer));
 							}
 
-							SCOPED_GPU_MASK(RHICmdList, FRHIGPUMask::FromIndex(GPUIndex));
+							RDG_GPU_MASK_SCOPE(GraphBuilder, FRHIGPUMask::FromIndex(GPUIndex));
 
 							if (GPUBatchedTileRequests.BatchedTilesDesc.Num() > 0)
 							{
-								FRDGBuilder GraphBuilder(RHICmdList);
-
 								FRDGTextureRef GBufferWorldPosition = GraphBuilder.RegisterExternalTexture(ScratchTilePoolGPU->PooledRenderTargets[0], TEXT("GBufferWorldPosition"));
 								FRDGTextureRef GBufferWorldNormal = GraphBuilder.RegisterExternalTexture(ScratchTilePoolGPU->PooledRenderTargets[1], TEXT("GBufferWorldNormal"));
 								FRDGTextureRef GBufferShadingNormal = GraphBuilder.RegisterExternalTexture(ScratchTilePoolGPU->PooledRenderTargets[2], TEXT("GBufferShadingNormal"));
@@ -1715,8 +1743,6 @@ void FLightmapRenderer::Finalize(FRHICommandListImmediate& RHICmdList)
 											FIntVector(GPUBatchedTileRequests.BatchedTilesDesc.Num() * 256, 1, 1));
 									}
 								}
-
-								GraphBuilder.Execute();
 							}
 						}
 					}
@@ -1728,8 +1754,6 @@ void FLightmapRenderer::Finalize(FRHICommandListImmediate& RHICmdList)
 
 	for (int32 SampleIndex = 0; SampleIndex < NumSamplesPerFrame; SampleIndex++)
 	{
-		FMemMark PerSampleMark(FMemStack::Get());
-
 		// Render shadow mask
 		{
 			TArray<FLightmapTileRequest> PendingShadowTileRequestsOnAllGPUs = PendingTileRequests.FilterByPredicate([NumShadowSamples = Scene->Settings->StationaryLightShadowSamples](const FLightmapTileRequest& Tile) { return !Tile.RenderState->IsTileShadowConverged(Tile.VirtualCoordinates, NumShadowSamples); });
@@ -1738,21 +1762,20 @@ void FLightmapRenderer::Finalize(FRHICommandListImmediate& RHICmdList)
 			{
 				for (uint32 GPUIndex = 0; GPUIndex < GNumExplicitGPUsForRendering; GPUIndex++)
 				{
-					SCOPED_GPU_MASK(RHICmdList, FRHIGPUMask::FromIndex(GPUIndex));
+					RDG_GPU_MASK_SCOPE(GraphBuilder, FRHIGPUMask::FromIndex(GPUIndex));
 
-					auto PendingShadowTileRequests = PendingShadowTileRequestsOnAllGPUs.FilterByPredicate(
+					auto& PendingShadowTileRequests = *GraphBuilder.AllocObject<TArray<FLightmapTileRequest>>(
+						PendingShadowTileRequestsOnAllGPUs.FilterByPredicate(
 						[GPUIndex](const FLightmapTileRequest& Tile)
 					{
 						uint32 AssignedGPUIndex = (Tile.RenderState->DistributionPrefixSum + Tile.RenderState->RetrieveTileStateIndex(Tile.VirtualCoordinates)) % GNumExplicitGPUsForRendering;
 						return AssignedGPUIndex == GPUIndex;
-					});
+					}));
 
 					if (PendingShadowTileRequests.Num() == 0) continue;
 
 					for (int ScratchLayerIndex = 0; ScratchLayerIndex < 3; ScratchLayerIndex++)
 					{
-						FRDGBuilder GraphBuilder(RHICmdList);
-
 						TResourceArray<FIntPoint> TilePositionsToClear;
 						for (auto& Tile : PendingShadowTileRequests)
 						{
@@ -1765,8 +1788,8 @@ void FLightmapRenderer::Finalize(FRHICommandListImmediate& RHICmdList)
 						FMultiTileClearCS::FParameters* Parameters = GraphBuilder.AllocParameters<FMultiTileClearCS::FParameters>();
 						Parameters->NumTiles = TilePositionsToClear.Num();
 						Parameters->TileSize = GPreviewLightmapPhysicalTileSize;
-						Parameters->TilePositions = TilePositionsBuffer.SRV;
-						Parameters->TilePool = ScratchTilePoolGPU->PooledRenderTargets[ScratchLayerIndex]->GetRenderTargetItem().UAV;
+						Parameters->TilePositions = HoldReference(GraphBuilder, TilePositionsBuffer.SRV);
+						Parameters->TilePool = ScratchTilePoolLayerUAVs[ScratchLayerIndex];
 
 						TShaderMapRef<FMultiTileClearCS> ComputeShader(GetGlobalShaderMap(GMaxRHIFeatureLevel));
 						FComputeShaderUtils::AddPass(
@@ -1775,17 +1798,8 @@ void FLightmapRenderer::Finalize(FRHICommandListImmediate& RHICmdList)
 							ComputeShader,
 							Parameters,
 							FComputeShaderUtils::GetGroupCount(FIntPoint(GPreviewLightmapPhysicalTileSize * TilePositionsToClear.Num(), GPreviewLightmapPhysicalTileSize), FComputeShaderUtils::kGolden2DGroupSize));
-
-						GraphBuilder.Execute();
-					}
-
-					for (int ScratchLayerIndex = 0; ScratchLayerIndex < 3; ScratchLayerIndex++)
-					{
-						RHICmdList.Transition(FRHITransitionInfo(ScratchTilePoolGPU->PooledRenderTargets[ScratchLayerIndex]->GetRenderTargetItem().UAV, ERHIAccess::UAVGraphics, ERHIAccess::UAVGraphics));
 					}
 
-					FRDGBuilder GraphBuilder(RHICmdList);
-
 					FRDGTextureRef GBufferWorldPosition = GraphBuilder.RegisterExternalTexture(ScratchTilePoolGPU->PooledRenderTargets[0], TEXT("GBufferWorldPosition"));
 					FRDGTextureRef GBufferWorldNormal = GraphBuilder.RegisterExternalTexture(ScratchTilePoolGPU->PooledRenderTargets[1], TEXT("GBufferWorldNormal"));
 					FRDGTextureRef GBufferShadingNormal = GraphBuilder.RegisterExternalTexture(ScratchTilePoolGPU->PooledRenderTargets[2], TEXT("GBufferShadingNormal"));
@@ -1801,7 +1815,7 @@ void FLightmapRenderer::Finalize(FRHICommandListImmediate& RHICmdList)
 					FVertexBufferRHIRef ChannelIndexBuffer;
 					FShaderResourceViewRHIRef ChannelIndexSRV;
 
-					TResourceArray<int32> LightSampleIndexArray;
+					auto& LightSampleIndexArray = *GraphBuilder.AllocObject<TResourceArray<int32>>();
 					FVertexBufferRHIRef LightSampleIndexBuffer;
 					FShaderResourceViewRHIRef LightSampleIndexSRV;
 
@@ -1962,54 +1976,58 @@ void FLightmapRenderer::Finalize(FRHICommandListImmediate& RHICmdList)
 
 					// Render GBuffer
 					{
-						FRHIRenderPassInfo RPInfo(FRHIRenderPassInfo::NoRenderTargets);
-						RHICmdList.BeginRenderPass(RPInfo, TEXT("LightmapGBuffer"));
+						auto* PassParameters = GraphBuilder.AllocParameters<FLightmapGBufferPassParameters>();
+						PassParameters->PassUniformBuffer = PassUniformBuffer;
 
-						for (int32 TileIndex = 0; TileIndex < PendingShadowTileRequests.Num(); TileIndex++)
+						GraphBuilder.AddPass(
+							RDG_EVENT_NAME("LightmapGBuffer"),
+							PassParameters,
+							ERDGPassFlags::Raster,
+							[this, ReferenceView = Scene->ReferenceView, &PendingShadowTileRequests, &LightSampleIndexArray](FRHICommandList& RHICmdList)
 						{
-							const FLightmapTileRequest& Tile = PendingShadowTileRequests[TileIndex];
-
-							RHICmdList.SetViewport(0, 0, 0.0f, GPreviewLightmapPhysicalTileSize, GPreviewLightmapPhysicalTileSize, 1.0f);
-
-							float ScaleX = Tile.RenderState->GetPaddedSizeInTiles().X * GPreviewLightmapVirtualTileSize * 1.0f / (1 << Tile.VirtualCoordinates.MipLevel) / GPreviewLightmapPhysicalTileSize;
-							float ScaleY = Tile.RenderState->GetPaddedSizeInTiles().Y * GPreviewLightmapVirtualTileSize * 1.0f / (1 << Tile.VirtualCoordinates.MipLevel) / GPreviewLightmapPhysicalTileSize;
-							float BiasX = (1.0f * (-Tile.VirtualCoordinates.Position.X * GPreviewLightmapVirtualTileSize) - (-GPreviewLightmapTileBorderSize)) / GPreviewLightmapPhysicalTileSize;
-							float BiasY = (1.0f * (-Tile.VirtualCoordinates.Position.Y * GPreviewLightmapVirtualTileSize) - (-GPreviewLightmapTileBorderSize)) / GPreviewLightmapPhysicalTileSize;
-
-							FVector4 VirtualTexturePhysicalTileCoordinateScaleAndBias = FVector4(ScaleX, ScaleY, BiasX, BiasY);
+							for (int32 TileIndex = 0; TileIndex < PendingShadowTileRequests.Num(); TileIndex++)
+							{
+								const FLightmapTileRequest& Tile = PendingShadowTileRequests[TileIndex];
 
-							TArray<FMeshBatch> MeshBatches = Tile.RenderState->GeometryInstanceRef.GetMeshBatchesForGBufferRendering(Tile.VirtualCoordinates);
+								RHICmdList.SetViewport(0, 0, 0.0f, GPreviewLightmapPhysicalTileSize, GPreviewLightmapPhysicalTileSize, 1.0f);
 
-							for (auto& MeshBatch : MeshBatches)
-							{
-								FMeshBatchElement& Element = MeshBatch.Elements[0];
+								float ScaleX = Tile.RenderState->GetPaddedSizeInTiles().X * GPreviewLightmapVirtualTileSize * 1.0f / (1 << Tile.VirtualCoordinates.MipLevel) / GPreviewLightmapPhysicalTileSize;
+								float ScaleY = Tile.RenderState->GetPaddedSizeInTiles().Y * GPreviewLightmapVirtualTileSize * 1.0f / (1 << Tile.VirtualCoordinates.MipLevel) / GPreviewLightmapPhysicalTileSize;
+								float BiasX = (1.0f * (-Tile.VirtualCoordinates.Position.X * GPreviewLightmapVirtualTileSize) - (-GPreviewLightmapTileBorderSize)) / GPreviewLightmapPhysicalTileSize;
+								float BiasY = (1.0f * (-Tile.VirtualCoordinates.Position.Y * GPreviewLightmapVirtualTileSize) - (-GPreviewLightmapTileBorderSize)) / GPreviewLightmapPhysicalTileSize;
 
-								Element.DynamicPrimitiveShaderDataIndex = Tile.RenderState->GeometryInstanceRef.GetElementId();
-							}
+								FVector4 VirtualTexturePhysicalTileCoordinateScaleAndBias = FVector4(ScaleX, ScaleY, BiasX, BiasY);
 
-							DrawDynamicMeshPass(
-								*Scene->ReferenceView, RHICmdList,
-								[
-									View = Scene->ReferenceView.Get(),
-									PassUniformBuffer,
-									MeshBatches,
-									VirtualTexturePhysicalTileCoordinateScaleAndBias,
-									RenderPassIndex = LightSampleIndexArray[TileIndex],
-									ScratchTilePoolOffset = ScratchTilePoolGPU->GetPositionFromLinearAddress(Tile.TileAddressInScratch) * GPreviewLightmapPhysicalTileSize
-								](FDynamicPassMeshDrawListContext* DynamicMeshPassContext)
-							{
-								FLightmapGBufferMeshProcessor MeshProcessor(nullptr, View, DynamicMeshPassContext, PassUniformBuffer, VirtualTexturePhysicalTileCoordinateScaleAndBias, RenderPassIndex, ScratchTilePoolOffset);
+								TArray<FMeshBatch> MeshBatches = Tile.RenderState->GeometryInstanceRef.GetMeshBatchesForGBufferRendering(Tile.VirtualCoordinates);
 
 								for (auto& MeshBatch : MeshBatches)
 								{
-									MeshProcessor.AddMeshBatch(MeshBatch, ~0ull, nullptr);
+									FMeshBatchElement& Element = MeshBatch.Elements[0];
+
+									Element.DynamicPrimitiveShaderDataIndex = Tile.RenderState->GeometryInstanceRef.GetElementId();
 								}
-							});
 
-							GPrimitiveIdVertexBufferPool.DiscardAll();
-						}
+								DrawDynamicMeshPass(
+									*ReferenceView, RHICmdList,
+									[
+										View = ReferenceView.Get(),
+										MeshBatches,
+										VirtualTexturePhysicalTileCoordinateScaleAndBias,
+										RenderPassIndex = LightSampleIndexArray[TileIndex],
+										ScratchTilePoolOffset = ScratchTilePoolGPU->GetPositionFromLinearAddress(Tile.TileAddressInScratch) * GPreviewLightmapPhysicalTileSize
+									](FDynamicPassMeshDrawListContext* DynamicMeshPassContext)
+								{
+									FLightmapGBufferMeshProcessor MeshProcessor(nullptr, View, DynamicMeshPassContext, VirtualTexturePhysicalTileCoordinateScaleAndBias, RenderPassIndex, ScratchTilePoolOffset);
 
-						RHICmdList.EndRenderPass();
+									for (auto& MeshBatch : MeshBatches)
+									{
+										MeshProcessor.AddMeshBatch(MeshBatch, ~0ull, nullptr);
+									}
+								});
+
+								GPrimitiveIdVertexBufferPool.DiscardAll();
+							}
+						});
 					}
 
 #if RHI_RAYTRACING
@@ -2039,7 +2057,7 @@ void FLightmapRenderer::Finalize(FRHICommandListImmediate& RHICmdList)
 								CreateInfo.ResourceArray = &GPUBatchedTileRequests.BatchedTilesDesc;
 
 								GPUBatchedTileRequests.BatchedTilesBuffer = RHICreateStructuredBuffer(sizeof(FGPUTileDescription), GPUBatchedTileRequests.BatchedTilesDesc.GetResourceDataSize(), BUF_Static | BUF_ShaderResource, CreateInfo);
-								GPUBatchedTileRequests.BatchedTilesSRV = RHICreateShaderResourceView(GPUBatchedTileRequests.BatchedTilesBuffer);
+								GPUBatchedTileRequests.BatchedTilesSRV = HoldReference(GraphBuilder, RHICreateShaderResourceView(GPUBatchedTileRequests.BatchedTilesBuffer));
 							}
 						}
 
@@ -2050,10 +2068,10 @@ void FLightmapRenderer::Finalize(FRHICommandListImmediate& RHICmdList)
 						FStationaryLightShadowTracingRGS::FParameters* PassParameters = GraphBuilder.AllocParameters<FStationaryLightShadowTracingRGS::FParameters>();
 						PassParameters->TLAS = Scene->RayTracingScene->GetShaderResourceView();
 						PassParameters->BatchedTiles = GPUBatchedTileRequests.BatchedTilesSRV;
-						PassParameters->LightTypeArray = LightTypeSRV;
-						PassParameters->ChannelIndexArray = ChannelIndexSRV;
-						PassParameters->LightSampleIndexArray = LightSampleIndexSRV;
-						PassParameters->LightShaderParametersArray = LightShaderParameterSRV;
+						PassParameters->LightTypeArray = HoldReference(GraphBuilder, LightTypeSRV);
+						PassParameters->ChannelIndexArray = HoldReference(GraphBuilder, ChannelIndexSRV);
+						PassParameters->LightSampleIndexArray = HoldReference(GraphBuilder, LightSampleIndexSRV);
+						PassParameters->LightShaderParametersArray = HoldReference(GraphBuilder, LightShaderParameterSRV);
 						PassParameters->GBufferWorldPosition = GBufferWorldPosition;
 						PassParameters->GBufferWorldNormal = GBufferWorldNormal;
 						PassParameters->GBufferShadingNormal = GBufferShadingNormal;
@@ -2074,10 +2092,8 @@ void FLightmapRenderer::Finalize(FRHICommandListImmediate& RHICmdList)
 
 							RHICmdList.RayTraceDispatch(PipelineState, RayGenerationShader.GetRayTracingShader(), RayTracingScene, GlobalResources, RayTracingResolution.X, RayTracingResolution.Y);
 						});
-
 					}
 #endif
-					GraphBuilder.Execute();
 				}
 			}
 		}
@@ -2118,7 +2134,10 @@ void FLightmapRenderer::Finalize(FRHICommandListImmediate& RHICmdList)
 			}
 		}
 
-		RHICmdList.TransferTextures(Params);
+		AddPass(GraphBuilder, [LocalParams = MoveTemp(Params)](FRHICommandListImmediate& RHICmdList)
+		{
+			RHICmdList.TransferTextures(LocalParams);
+		});
 	}
 
 	// Output from working set to VT layers
@@ -2149,8 +2168,6 @@ void FLightmapRenderer::Finalize(FRHICommandListImmediate& RHICmdList)
 		}
 
 		{
-			FRDGBuilder GraphBuilder(RHICmdList);
-
 			FRDGTextureRef IrradianceAndSampleCount = GraphBuilder.RegisterExternalTexture(LightmapTilePoolGPU.PooledRenderTargets[0], TEXT("IrradianceAndSampleCount"));
 			FRDGTextureRef SHDirectionality = GraphBuilder.RegisterExternalTexture(LightmapTilePoolGPU.PooledRenderTargets[1], TEXT("SHDirectionality"));
 			FRDGTextureRef ShadowMask = GraphBuilder.RegisterExternalTexture(LightmapTilePoolGPU.PooledRenderTargets[2], TEXT("ShadowMask"));
@@ -2214,12 +2231,13 @@ void FLightmapRenderer::Finalize(FRHICommandListImmediate& RHICmdList)
 					PassParameters,
 					FComputeShaderUtils::GetGroupCount(RayTracingResolution, FComputeShaderUtils::kGolden2DGroupSize));
 			}
-
-			GraphBuilder.Execute();
 		}
 	}
 
-	Scene->DestroyRayTracingScene();
+	AddPass(GraphBuilder, [this](FRHICommandList&)
+	{
+		Scene->DestroyRayTracingScene();
+	});
 
 	// Perform readback on any potential converged tiles
 	{
@@ -2248,7 +2266,7 @@ void FLightmapRenderer::Finalize(FRHICommandListImmediate& RHICmdList)
 
 			for (uint32 GPUIndex = 0; GPUIndex < GNumExplicitGPUsForRendering; GPUIndex++)
 			{
-				SCOPED_GPU_MASK(RHICmdList, FRHIGPUMask::FromIndex(GPUIndex));
+				RDG_GPU_MASK_SCOPE(GraphBuilder, FRHIGPUMask::FromIndex(GPUIndex));
 
 				auto ConvergedTileRequestsOnCurrentGPU = ConvergedTileRequests.FilterByPredicate(
 					[GPUIndex](const FLightmapTileRequest& Tile)
@@ -2313,15 +2331,13 @@ void FLightmapRenderer::Finalize(FRHICommandListImmediate& RHICmdList)
 				CreateInfo.ResourceArray = &GPUBatchedTileRequests.BatchedTilesDesc;
 
 				GPUBatchedTileRequests.BatchedTilesBuffer = RHICreateStructuredBuffer(sizeof(FGPUTileDescription), GPUBatchedTileRequests.BatchedTilesDesc.GetResourceDataSize(), BUF_Static | BUF_ShaderResource, CreateInfo);
-				GPUBatchedTileRequests.BatchedTilesSRV = RHICreateShaderResourceView(GPUBatchedTileRequests.BatchedTilesBuffer);
+				GPUBatchedTileRequests.BatchedTilesSRV = HoldReference(GraphBuilder, RHICreateShaderResourceView(GPUBatchedTileRequests.BatchedTilesBuffer));
 
 
 				FIntPoint DispatchResolution;
 				DispatchResolution.X = GPreviewLightmapPhysicalTileSize * GPUBatchedTileRequests.BatchedTilesDesc.Num();
 				DispatchResolution.Y = GPreviewLightmapPhysicalTileSize;
 
-				FRDGBuilder GraphBuilder(RHICmdList);
-
 				FRDGTextureRef IrradianceAndSampleCount = GraphBuilder.RegisterExternalTexture(LightmapTilePoolGPU.PooledRenderTargets[0], TEXT("IrradianceAndSampleCount"));
 				FRDGTextureRef SHDirectionality = GraphBuilder.RegisterExternalTexture(LightmapTilePoolGPU.PooledRenderTargets[1], TEXT("SHDirectionality"));
 				FRDGTextureRef ShadowMask = GraphBuilder.RegisterExternalTexture(LightmapTilePoolGPU.PooledRenderTargets[2], TEXT("ShadowMask"));
@@ -2356,11 +2372,9 @@ void FLightmapRenderer::Finalize(FRHICommandListImmediate& RHICmdList)
 						FComputeShaderUtils::GetGroupCount(DispatchResolution, FComputeShaderUtils::kGolden2DGroupSize));
 				}
 
-				GraphBuilder.Execute();
-
-				LightmapReadbackGroup.StagingHQLayer0Readback->EnqueueCopy(RHICmdList, LightmapReadbackGroup.ReadbackTilePoolGPU->PooledRenderTargets[0]->GetRenderTargetItem().TargetableTexture);
-				LightmapReadbackGroup.StagingHQLayer1Readback->EnqueueCopy(RHICmdList, LightmapReadbackGroup.ReadbackTilePoolGPU->PooledRenderTargets[1]->GetRenderTargetItem().TargetableTexture);
-				LightmapReadbackGroup.StagingShadowMaskReadback->EnqueueCopy(RHICmdList, LightmapReadbackGroup.ReadbackTilePoolGPU->PooledRenderTargets[2]->GetRenderTargetItem().TargetableTexture);
+				AddEnqueueCopyPass(GraphBuilder, LightmapReadbackGroup.StagingHQLayer0Readback.Get(), StagingHQLayer0);
+				AddEnqueueCopyPass(GraphBuilder, LightmapReadbackGroup.StagingHQLayer1Readback.Get(), StagingHQLayer1);
+				AddEnqueueCopyPass(GraphBuilder, LightmapReadbackGroup.StagingShadowMaskReadback.Get(), StagingShadowMask);
 
 				OngoingReadbacks.Emplace(ReadbackGroupToUse);
 			}
@@ -2742,7 +2756,7 @@ void FLightmapRenderer::BackgroundTick()
 				TexCreate_ShaderResource | TexCreate_RenderTargetable | TexCreate_UAV,
 				false);
 
-			GRenderTargetPool.FindFreeElement(RHICmdList, Desc, OutputTileAtlas, *FString::Printf(TEXT("BackgroundTilePoolGPU")));
+			GRenderTargetPool.FindFreeElement(RHICmdList, Desc, OutputTileAtlas, TEXT("BackgroundTilePoolGPU"));
 
 			ensure(OutputTileAtlas.IsValid());
 
@@ -2859,7 +2873,12 @@ void FLightmapRenderer::BackgroundTick()
 	bInsideBackgroundTick = true;
 
 	// Render lightmap tiles
-	Finalize(RHICmdList);
+	{
+		FMemMark MemMark(FMemStack::Get());
+		FRDGBuilder GraphBuilder(RHICmdList);
+		Finalize(GraphBuilder);
+		GraphBuilder.Execute();
+	}
 
 	bInsideBackgroundTick = false;
 
diff --git a/Engine/Plugins/Experimental/GPULightmass/Source/GPULightmass/Private/LightmapRenderer.h b/Engine/Plugins/Experimental/GPULightmass/Source/GPULightmass/Private/LightmapRenderer.h
--- a/Engine/Plugins/Experimental/GPULightmass/Source/GPULightmass/Private/LightmapRenderer.h
+++ b/Engine/Plugins/Experimental/GPULightmass/Source/GPULightmass/Private/LightmapRenderer.h
@@ -108,7 +108,7 @@ public:
 
 	void AddRequest(FLightmapTileRequest TileRequest);
 
-	virtual void Finalize(FRHICommandListImmediate& RHICmdList) override;
+	virtual void Finalize(FRDGBuilder& GraphBuilder) override;
 
 	virtual ~FLightmapRenderer();
 
diff --git a/Engine/Plugins/Experimental/GPULightmass/Source/GPULightmass/Private/Scene/Scene.h b/Engine/Plugins/Experimental/GPULightmass/Source/GPULightmass/Private/Scene/Scene.h
--- a/Engine/Plugins/Experimental/GPULightmass/Source/GPULightmass/Private/Scene/Scene.h
+++ b/Engine/Plugins/Experimental/GPULightmass/Source/GPULightmass/Private/Scene/Scene.h
@@ -84,7 +84,7 @@ public:
 
 	FRayTracingSceneRHIRef RayTracingScene;
 	FRayTracingPipelineState* RayTracingPipelineState;
-	TUniquePtr<FViewInfo> ReferenceView;
+	TSharedPtr<FViewInfo> ReferenceView;
 
 	TUniquePtr<FCachedRayTracingSceneData> CachedRayTracingScene;
 
diff --git a/Engine/Plugins/Experimental/VirtualHeightfieldMesh/Source/VirtualHeightfieldMeshEditor/Private/HeightfieldMinMaxTextureBuild.cpp b/Engine/Plugins/Experimental/VirtualHeightfieldMesh/Source/VirtualHeightfieldMeshEditor/Private/HeightfieldMinMaxTextureBuild.cpp
--- a/Engine/Plugins/Experimental/VirtualHeightfieldMesh/Source/VirtualHeightfieldMeshEditor/Private/HeightfieldMinMaxTextureBuild.cpp
+++ b/Engine/Plugins/Experimental/VirtualHeightfieldMesh/Source/VirtualHeightfieldMeshEditor/Private/HeightfieldMinMaxTextureBuild.cpp
@@ -213,12 +213,13 @@ namespace VirtualHeightfieldMesh
 					Desc.PageDescs[0].UVRange = UVRange;
 					Desc.PageDescs[0].vLevel = MipLevel;
 
-					RenderPages(RHICmdList, Desc);
-
-					// Downsample page to texel in output
 					FMemMark Mark(FMemStack::Get());
 					FRDGBuilder GraphBuilder(RHICmdList);
 
+					RenderPages(GraphBuilder, Desc);
+
+					// Downsample page to texel in output
+
 					FRDGTextureRef SrcTexture = GraphBuilder.RegisterExternalTexture(RenderTileResources.GetTileRenderTarget());
 					FRDGTextureRef DstTexture = GraphBuilder.RegisterExternalTexture(RenderTileResources.GetFinalRenderTarget());
 					FRDGTextureUAVRef DstTextureUAV = GraphBuilder.CreateUAV(DstTexture);
@@ -240,13 +241,12 @@ namespace VirtualHeightfieldMesh
 				FRDGBuilder GraphBuilder(RHICmdList);
 
 				FRDGTextureRef Texture = GraphBuilder.RegisterExternalTexture(RenderTileResources.GetFinalRenderTarget());
+				GraphBuilder.SetTextureAccessFinal(Texture, ERHIAccess::CopySrc);
 
 				GenerateMinMaxTextureMips(GraphBuilder, Texture, FIntPoint(NumTilesX, NumTilesY), NumMips);
 
 				GraphBuilder.Execute();
 
-				RHICmdList.Transition(FRHITransitionInfo(RenderTileResources.GetFinalRenderTarget()->GetRenderTargetItem().ShaderResourceTexture, ERHIAccess::WritableMask, ERHIAccess::CopySrc));
-
 				for (int32 MipLevel = 0; MipLevel < NumMips; MipLevel++)
 				{
 					FRHICopyTextureInfo CopyInfo;
diff --git a/Engine/Source/Editor/VirtualTexturingEditor/Private/RuntimeVirtualTextureBuildStreamingMips.cpp b/Engine/Source/Editor/VirtualTexturingEditor/Private/RuntimeVirtualTextureBuildStreamingMips.cpp
--- a/Engine/Source/Editor/VirtualTexturingEditor/Private/RuntimeVirtualTextureBuildStreamingMips.cpp
+++ b/Engine/Source/Editor/VirtualTexturingEditor/Private/RuntimeVirtualTextureBuildStreamingMips.cpp
@@ -9,6 +9,7 @@
 #include "RendererInterface.h"
 #include "RenderTargetPool.h"
 #include "SceneInterface.h"
+#include "RenderGraphBuilder.h"
 #include "VT/RuntimeVirtualTexture.h"
 #include "VT/RuntimeVirtualTextureRender.h"
 #include "VT/VirtualTextureBuilder.h"
@@ -244,7 +245,12 @@ namespace RuntimeVirtualTexture
 					Desc.PageDescs[0].UVRange = UVRange;
 					Desc.PageDescs[0].vLevel = RenderLevel;
 
-					RuntimeVirtualTexture::RenderPages(RHICmdList, Desc);
+					{
+						FMemMark MemMark(FMemStack::Get());
+						FRDGBuilder GraphBuilder(RHICmdList);
+						RuntimeVirtualTexture::RenderPages(GraphBuilder, Desc);
+						GraphBuilder.Execute();
+					}
 
 					// Transition render targets for copying
 					for (int32 Layer = 0; Layer < NumLayers; Layer++)
diff --git a/Engine/Source/Editor/VirtualTexturingEditor/Private/RuntimeVirtualTextureThumbnailRenderer.cpp b/Engine/Source/Editor/VirtualTexturingEditor/Private/RuntimeVirtualTextureThumbnailRenderer.cpp
--- a/Engine/Source/Editor/VirtualTexturingEditor/Private/RuntimeVirtualTextureThumbnailRenderer.cpp
+++ b/Engine/Source/Editor/VirtualTexturingEditor/Private/RuntimeVirtualTextureThumbnailRenderer.cpp
@@ -5,6 +5,7 @@
 #include "Components/RuntimeVirtualTextureComponent.h"
 #include "MaterialShared.h"
 #include "RenderingThread.h"
+#include "RenderGraphBuilder.h"
 #include "SceneInterface.h"
 #include "VT/RuntimeVirtualTexture.h"
 #include "VT/RuntimeVirtualTextureRender.h"
@@ -96,6 +97,9 @@ void URuntimeVirtualTextureThumbnailRenderer::Draw(UObject* Object, int32 X, int
 		Desc.PageDescs[0].UVRange = FBox2D(FVector2D(0, 0), FVector2D(1, 1));
 		Desc.PageDescs[0].vLevel = MaxLevel;
 
-		RuntimeVirtualTexture::RenderPages(RHICmdList, Desc);
+		FMemMark MemMark(FMemStack::Get());
+		FRDGBuilder GraphBuilder(RHICmdList);
+		RuntimeVirtualTexture::RenderPages(GraphBuilder, Desc);
+		GraphBuilder.Execute();
 	});
 }
diff --git a/Engine/Source/Runtime/Engine/Private/Texture2D.cpp b/Engine/Source/Runtime/Engine/Private/Texture2D.cpp
--- a/Engine/Source/Runtime/Engine/Private/Texture2D.cpp
+++ b/Engine/Source/Runtime/Engine/Private/Texture2D.cpp
@@ -17,6 +17,7 @@
 #include "UObject/LinkerLoad.h"
 #include "UObject/CoreRedirects.h"
 #include "RenderUtils.h"
+#include "RenderGraphBuilder.h"
 #include "ContentStreaming.h"
 #include "EngineUtils.h"
 #include "DeviceProfiles/DeviceProfile.h"
@@ -1351,12 +1352,16 @@ void FVirtualTexture2DResource::InitializeEditorResources(IVirtualTexture* InVir
 			}
 		}
 
-		for (IVirtualTextureFinalizer* Finalizer : Finalizers)
 		{
-			Finalizer->Finalize(RHICommandList);
+			FMemMark Mark(FMemStack::Get());
+			FRDGBuilder GraphBuilder(RHICommandList);
+			for (IVirtualTextureFinalizer* Finalizer : Finalizers)
+			{
+				Finalizer->Finalize(GraphBuilder);
+			}
+			GraphBuilder.Execute();
 		}
 
-		
 		if (MipWidthInTiles * TileSizeInPixels != MipWidth || MipHeightInTiles * TileSizeInPixels != MipHeight)
 		{
 			// Logical dimensions of mip image may be smaller than tile size (in this case tile will contain mirrored/wrapped padding)
diff --git a/Engine/Source/Runtime/Engine/Private/VT/VirtualTextureUploadCache.cpp b/Engine/Source/Runtime/Engine/Private/VT/VirtualTextureUploadCache.cpp
--- a/Engine/Source/Runtime/Engine/Private/VT/VirtualTextureUploadCache.cpp
+++ b/Engine/Source/Runtime/Engine/Private/VT/VirtualTextureUploadCache.cpp
@@ -3,6 +3,7 @@
 #include "VirtualTextureUploadCache.h"
 #include "VirtualTextureChunkManager.h"
 #include "RHI.h"
+#include "RenderGraphBuilder.h"
 
 // allow uploading CPU buffer directly to GPU texture
 // this is slow under D3D11
@@ -83,12 +84,14 @@ int32 FVirtualTextureUploadCache::GetOrCreatePoolIndex(EPixelFormat InFormat, ui
 	return PoolIndex;
 }
 
-void FVirtualTextureUploadCache::Finalize(FRHICommandListImmediate& RHICmdList)
+void FVirtualTextureUploadCache::Finalize(FRDGBuilder& GraphBuilder)
 {
 	SCOPE_CYCLE_COUNTER(STAT_VTP_FlushUpload)
 
 	check(IsInRenderingThread());
 
+	FRHICommandListImmediate& RHICmdList = GraphBuilder.RHICmdList;
+
 	// Multi-GPU support : May be ineffecient for AFR.
 	SCOPED_GPU_MASK(RHICmdList, FRHIGPUMask::All());
 
diff --git a/Engine/Source/Runtime/Engine/Private/VT/VirtualTextureUploadCache.h b/Engine/Source/Runtime/Engine/Private/VT/VirtualTextureUploadCache.h
--- a/Engine/Source/Runtime/Engine/Private/VT/VirtualTextureUploadCache.h
+++ b/Engine/Source/Runtime/Engine/Private/VT/VirtualTextureUploadCache.h
@@ -35,7 +35,7 @@ public:
 	virtual ~FVirtualTextureUploadCache();
 
 	// IVirtualTextureFinalizer
-	virtual void Finalize(FRHICommandListImmediate& RHICmdList) override;
+	virtual void Finalize(FRDGBuilder& GraphBuilder) override;
 
 	// FRenderResource
 	virtual void ReleaseRHI() override;
diff --git a/Engine/Source/Runtime/RenderCore/Public/VirtualTexturing.h b/Engine/Source/Runtime/RenderCore/Public/VirtualTexturing.h
--- a/Engine/Source/Runtime/RenderCore/Public/VirtualTexturing.h
+++ b/Engine/Source/Runtime/RenderCore/Public/VirtualTexturing.h
@@ -12,6 +12,7 @@ class FRHICommandListImmediate;
 class FRHIShaderResourceView;
 class FRHITexture;
 class FRHIUnorderedAccessView;
+class FRDGBuilder;
 
 union FVirtualTextureProducerHandle
 {
@@ -150,7 +151,10 @@ typedef void (FVTProducerDestroyedFunction)(const FVirtualTextureProducerHandle&
 class IVirtualTextureFinalizer
 {
 public:
-	virtual void Finalize(FRHICommandListImmediate& RHICmdList) = 0;
+	virtual void Finalize(FRDGBuilder& GraphBuilder) = 0;
+
+	UE_DEPRECATED(5.0, "This method has been refactored to use an FRDGBuilder instead.")
+	virtual void Finalize(FRHICommandListImmediate& RHICmdList) {}
 };
 
 enum class EVTRequestPageStatus
diff --git a/Engine/Source/Runtime/Renderer/Private/DeferredShadingRenderer.cpp b/Engine/Source/Runtime/Renderer/Private/DeferredShadingRenderer.cpp
--- a/Engine/Source/Runtime/Renderer/Private/DeferredShadingRenderer.cpp
+++ b/Engine/Source/Runtime/Renderer/Private/DeferredShadingRenderer.cpp
@@ -1477,10 +1477,14 @@ void FDeferredShadingSceneRenderer::Render(FRHICommandListImmediate& RHICmdList)
 	const bool bUseVirtualTexturing = UseVirtualTexturing(FeatureLevel);
 	if (bUseVirtualTexturing)
 	{
-		SCOPED_GPU_STAT(RHICmdList, VirtualTextureUpdate);
-		// AllocateResources needs to be called before RHIBeginScene
-		FVirtualTextureSystem::Get().AllocateResources(RHICmdList, FeatureLevel);
-		FVirtualTextureSystem::Get().CallPendingCallbacks();
+		FRDGBuilder GraphBuilder(RHICmdList);
+		{
+			SCOPED_GPU_STAT(RHICmdList, VirtualTextureUpdate);
+			// AllocateResources needs to be called before RHIBeginScene
+			FVirtualTextureSystem::Get().AllocateResources(GraphBuilder, FeatureLevel);
+			FVirtualTextureSystem::Get().CallPendingCallbacks();
+		}
+		GraphBuilder.Execute();
 	}
 
 	// Nanite materials do not currently support most debug view modes.
@@ -1554,15 +1558,22 @@ void FDeferredShadingSceneRenderer::Render(FRHICommandListImmediate& RHICmdList)
 
 		if (bUseVirtualTexturing)
 		{
-			SCOPED_GPU_STAT(RHICmdList, VirtualTextureUpdate);
-			FVirtualTextureSystem::Get().Update(RHICmdList, FeatureLevel, Scene);
-
-			// Clear virtual texture feedback to default value
-			FUnorderedAccessViewRHIRef FeedbackUAV = SceneContext.GetVirtualTextureFeedbackUAV();
-			RHICmdList.Transition(FRHITransitionInfo(FeedbackUAV, ERHIAccess::Unknown, ERHIAccess::UAVCompute));
-			RHICmdList.ClearUAVUint(FeedbackUAV, FUintVector4(~0u, ~0u, ~0u, ~0u));
-			RHICmdList.Transition(FRHITransitionInfo(FeedbackUAV, ERHIAccess::UAVCompute, ERHIAccess::UAVCompute));
-			RHICmdList.BeginUAVOverlap(FeedbackUAV);
+			FRDGBuilder GraphBuilder(RHICmdList);
+			{
+				RDG_GPU_STAT_SCOPE(GraphBuilder, VirtualTextureUpdate);
+				FVirtualTextureSystem::Get().Update(GraphBuilder, FeatureLevel, Scene);
+
+				AddPass(GraphBuilder, [&SceneContext](FRHICommandList& RHICmdList)
+				{
+					// Clear virtual texture feedback to default value
+					FUnorderedAccessViewRHIRef FeedbackUAV = SceneContext.GetVirtualTextureFeedbackUAV();
+					RHICmdList.Transition(FRHITransitionInfo(FeedbackUAV, ERHIAccess::Unknown, ERHIAccess::UAVCompute));
+					RHICmdList.ClearUAVUint(FeedbackUAV, FUintVector4(~0u, ~0u, ~0u, ~0u));
+					RHICmdList.Transition(FRHITransitionInfo(FeedbackUAV, ERHIAccess::UAVCompute, ERHIAccess::UAVCompute));
+					RHICmdList.BeginUAVOverlap(FeedbackUAV);
+				});
+			}
+			GraphBuilder.Execute();
 		}
 
 		if (GDoPrepareDistanceFieldSceneAfterRHIFlush && (GRHINeedsExtraDeletionLatency || !GRHICommandList.Bypass()))
diff --git a/Engine/Source/Runtime/Renderer/Private/MobileShadingRenderer.cpp b/Engine/Source/Runtime/Renderer/Private/MobileShadingRenderer.cpp
--- a/Engine/Source/Runtime/Renderer/Private/MobileShadingRenderer.cpp
+++ b/Engine/Source/Runtime/Renderer/Private/MobileShadingRenderer.cpp
@@ -543,10 +543,14 @@ void FMobileSceneRenderer::Render(FRHICommandListImmediate& RHICmdList)
 	const bool bUseVirtualTexturing = UseVirtualTexturing(ViewFeatureLevel);
 	if (bUseVirtualTexturing)
 	{
-		SCOPED_GPU_STAT(RHICmdList, VirtualTextureUpdate);
-		// AllocateResources needs to be called before RHIBeginScene
-		FVirtualTextureSystem::Get().AllocateResources(RHICmdList, ViewFeatureLevel);
-		FVirtualTextureSystem::Get().CallPendingCallbacks();
+		FRDGBuilder GraphBuilder(RHICmdList);
+		{
+			SCOPED_GPU_STAT(RHICmdList, VirtualTextureUpdate);
+			// AllocateResources needs to be called before RHIBeginScene
+			FVirtualTextureSystem::Get().AllocateResources(GraphBuilder, FeatureLevel);
+			FVirtualTextureSystem::Get().CallPendingCallbacks();
+		}
+		GraphBuilder.Execute();
 	}
 
 	//make sure all the targets we're going to use will be safely writable.
@@ -576,13 +580,22 @@ void FMobileSceneRenderer::Render(FRHICommandListImmediate& RHICmdList)
 
 	if (bUseVirtualTexturing)
 	{
-		SCOPED_GPU_STAT(RHICmdList, VirtualTextureUpdate);
-		FVirtualTextureSystem::Get().Update(RHICmdList, ViewFeatureLevel, Scene);
-		// Clear virtual texture feedback to default value
-		FUnorderedAccessViewRHIRef FeedbackUAV = SceneContext.GetVirtualTextureFeedbackUAV();
-		RHICmdList.Transition(FRHITransitionInfo(FeedbackUAV, ERHIAccess::SRVMask, ERHIAccess::UAVMask));
-		RHICmdList.ClearUAVUint(FeedbackUAV, FUintVector4(~0u, ~0u, ~0u, ~0u));
-		RHICmdList.Transition(FRHITransitionInfo(FeedbackUAV, ERHIAccess::UAVMask, ERHIAccess::UAVMask));
+		FRDGBuilder GraphBuilder(RHICmdList);
+		{
+			RDG_GPU_STAT_SCOPE(GraphBuilder, VirtualTextureUpdate);
+			FVirtualTextureSystem::Get().Update(GraphBuilder, FeatureLevel, Scene);
+
+			AddPass(GraphBuilder, [&SceneContext](FRHICommandList& RHICmdList)
+			{
+				// Clear virtual texture feedback to default value
+				FUnorderedAccessViewRHIRef FeedbackUAV = SceneContext.GetVirtualTextureFeedbackUAV();
+				RHICmdList.Transition(FRHITransitionInfo(FeedbackUAV, ERHIAccess::Unknown, ERHIAccess::UAVCompute));
+				RHICmdList.ClearUAVUint(FeedbackUAV, FUintVector4(~0u, ~0u, ~0u, ~0u));
+				RHICmdList.Transition(FRHITransitionInfo(FeedbackUAV, ERHIAccess::UAVCompute, ERHIAccess::UAVCompute));
+				RHICmdList.BeginUAVOverlap(FeedbackUAV);
+			});
+		}
+		GraphBuilder.Execute();
 	}
 
 	FSortedLightSetSceneInfo SortedLightSet;
diff --git a/Engine/Source/Runtime/Renderer/Private/SceneHitProxyRendering.cpp b/Engine/Source/Runtime/Renderer/Private/SceneHitProxyRendering.cpp
--- a/Engine/Source/Runtime/Renderer/Private/SceneHitProxyRendering.cpp
+++ b/Engine/Source/Runtime/Renderer/Private/SceneHitProxyRendering.cpp
@@ -222,8 +222,14 @@ void InitHitProxyRender(FRHICommandListImmediate& RHICmdList, const FSceneRender
 	// Ensure VirtualTexture resources are allocated
 	if (UseVirtualTexturing(FeatureLevel))
 	{
-		FVirtualTextureSystem::Get().AllocateResources(RHICmdList, FeatureLevel);
-		FVirtualTextureSystem::Get().CallPendingCallbacks();
+		FRDGBuilder GraphBuilder(RHICmdList);
+		{
+			SCOPED_GPU_STAT(RHICmdList, VirtualTextureUpdate);
+			// AllocateResources needs to be called before RHIBeginScene
+			FVirtualTextureSystem::Get().AllocateResources(GraphBuilder, FeatureLevel);
+			FVirtualTextureSystem::Get().CallPendingCallbacks();
+		}
+		GraphBuilder.Execute();
 	}
 
 	// Initialize global system textures (pass-through if already initialized).
diff --git a/Engine/Source/Runtime/Renderer/Private/SceneRendering.cpp b/Engine/Source/Runtime/Renderer/Private/SceneRendering.cpp
--- a/Engine/Source/Runtime/Renderer/Private/SceneRendering.cpp
+++ b/Engine/Source/Runtime/Renderer/Private/SceneRendering.cpp
@@ -4127,7 +4127,10 @@ void FRendererModule::RequestVirtualTextureTilesForRegion(IAllocatedVirtualTextu
 
 void FRendererModule::LoadPendingVirtualTextureTiles(FRHICommandListImmediate& RHICmdList, ERHIFeatureLevel::Type FeatureLevel)
 {
-	FVirtualTextureSystem::Get().LoadPendingTiles(RHICmdList, FeatureLevel);
+	FMemMark MemMark(FMemStack::Get());
+	FRDGBuilder GraphBuilder(RHICmdList);
+	FVirtualTextureSystem::Get().LoadPendingTiles(GraphBuilder, FeatureLevel);
+	GraphBuilder.Execute();
 }
 
 void FRendererModule::FlushVirtualTextureCache()
diff --git a/Engine/Source/Runtime/Renderer/Private/VT/RuntimeVirtualTextureProducer.cpp b/Engine/Source/Runtime/Renderer/Private/VT/RuntimeVirtualTextureProducer.cpp
--- a/Engine/Source/Runtime/Renderer/Private/VT/RuntimeVirtualTextureProducer.cpp
+++ b/Engine/Source/Runtime/Renderer/Private/VT/RuntimeVirtualTextureProducer.cpp
@@ -61,7 +61,7 @@ void FRuntimeVirtualTextureFinalizer::AddTile(FTileEntry& Tile)
 	Tiles.Add(Tile);
 }
 
-void FRuntimeVirtualTextureFinalizer::Finalize(FRHICommandListImmediate& RHICmdList)
+void FRuntimeVirtualTextureFinalizer::Finalize(FRDGBuilder& GraphBuilder)
 {
 	RuntimeVirtualTexture::FRenderPageBatchDesc RenderPageBatchDesc;
 	RenderPageBatchDesc.Scene = Scene->GetRenderScene();
@@ -115,7 +115,7 @@ void FRuntimeVirtualTextureFinalizer::Finalize(FRHICommandListImmediate& RHICmdL
 		if (++BatchSize == RuntimeVirtualTexture::EMaxRenderPageBatch || bBreakBatchForTextures)
 		{
 			RenderPageBatchDesc.NumPageDescs = BatchSize;
-			RuntimeVirtualTexture::RenderPages(RHICmdList, RenderPageBatchDesc);
+			RuntimeVirtualTexture::RenderPages(GraphBuilder, RenderPageBatchDesc);
 			BatchSize = 0;
 		}
 
@@ -132,7 +132,7 @@ void FRuntimeVirtualTextureFinalizer::Finalize(FRHICommandListImmediate& RHICmdL
 	if (BatchSize > 0)
 	{
 		RenderPageBatchDesc.NumPageDescs = BatchSize;
-		RuntimeVirtualTexture::RenderPages(RHICmdList, RenderPageBatchDesc);
+		RuntimeVirtualTexture::RenderPages(GraphBuilder, RenderPageBatchDesc);
 	}
 
 	Tiles.SetNumUnsafeInternal(0);
diff --git a/Engine/Source/Runtime/Renderer/Private/VT/RuntimeVirtualTextureProducer.h b/Engine/Source/Runtime/Renderer/Private/VT/RuntimeVirtualTextureProducer.h
--- a/Engine/Source/Runtime/Renderer/Private/VT/RuntimeVirtualTextureProducer.h
+++ b/Engine/Source/Runtime/Renderer/Private/VT/RuntimeVirtualTextureProducer.h
@@ -44,7 +44,7 @@ public:
 	void AddTile(FTileEntry& Tile);
 
 	//~ Begin IVirtualTextureFinalizer Interface.
-	virtual void Finalize(FRHICommandListImmediate& RHICmdList) override;
+	virtual void Finalize(FRDGBuilder& GraphBuilder) override;
 	//~ End IVirtualTextureFinalizer Interface.
 
 private:
diff --git a/Engine/Source/Runtime/Renderer/Private/VT/RuntimeVirtualTextureRender.cpp b/Engine/Source/Runtime/Renderer/Private/VT/RuntimeVirtualTextureRender.cpp
--- a/Engine/Source/Runtime/Renderer/Private/VT/RuntimeVirtualTextureRender.cpp
+++ b/Engine/Source/Runtime/Renderer/Private/VT/RuntimeVirtualTextureRender.cpp
@@ -804,8 +804,6 @@ namespace RuntimeVirtualTexture
 			ERDGPassFlags::Copy | ERDGPassFlags::NeverCull,
 			[InputTexture, OutputTexture, CopyInfo](FRHICommandList& RHICmdList)
 			{
-				InputTexture->MarkResourceAsUsed();
-
 				RHICmdList.Transition(FRHITransitionInfo(OutputTexture, ERHIAccess::SRVMask, ERHIAccess::CopyDest));
 				RHICmdList.CopyTexture(InputTexture->GetRHI(), OutputTexture, CopyInfo);
 				RHICmdList.Transition(FRHITransitionInfo(OutputTexture, ERHIAccess::CopyDest, ERHIAccess::SRVMask));
@@ -964,7 +962,7 @@ namespace RuntimeVirtualTexture
 	}
 
 	void RenderPage(
-		FRHICommandListImmediate& RHICmdList,
+		FRDGBuilder& GraphBuilder,
 		FScene* Scene,
 		uint32 RuntimeVirtualTextureMask,
 		ERuntimeVirtualTextureMaterialType MaterialType,
@@ -986,14 +984,14 @@ namespace RuntimeVirtualTexture
 		uint8 MaxLevel,
 		ERuntimeVirtualTextureDebugType DebugType)
 	{
-		SCOPED_DRAW_EVENT(RHICmdList, VirtualTextureDynamicCache);
+		RDG_EVENT_SCOPE(GraphBuilder, "VirtualTextureDynamicCache");
 
 		// Initialize a temporary view required for the material render pass
 		//todo[vt]: Some of this, such as ViewRotationMatrix, can be computed once in the Finalizer and passed down.
 		//todo[vt]: Have specific shader variations and setup for different output texture configs
 		FSceneViewFamily::ConstructionValues ViewFamilyInit(nullptr, nullptr, FEngineShowFlags(ESFIM_Game));
 		ViewFamilyInit.SetWorldTimes(0.0f, 0.0f, 0.0f);
-		FSceneViewFamily ViewFamily(ViewFamilyInit);
+		FSceneViewFamily& ViewFamily = *GraphBuilder.AllocObject<FSceneViewFamily>(ViewFamilyInit);
 
 		FSceneViewInitOptions ViewInitOptions;
 		ViewInitOptions.ViewFamily = &ViewFamily;
@@ -1031,11 +1029,11 @@ namespace RuntimeVirtualTexture
 		ViewInitOptions.BackgroundColor = FLinearColor::Black;
 		ViewInitOptions.OverlayColor = FLinearColor::White;
 
-		FViewInfo ViewInfo(ViewInitOptions);
+		FViewInfo& ViewInfo = *GraphBuilder.AllocObject<FViewInfo>(ViewInitOptions);
 		FViewInfo* View = &ViewInfo;
 		ViewFamily.Views.Add(View);
 
-		FSceneRenderTargets& SceneContext = FSceneRenderTargets::Get(RHICmdList);
+		FSceneRenderTargets& SceneContext = FSceneRenderTargets::Get(GraphBuilder.RHICmdList);
 		View->bIsVirtualTexture = true;
 		View->ViewRect = View->UnconstrainedViewRect;
 		View->CachedViewUniformShaderParameters = MakeUnique<FViewUniformShaderParameters>();
@@ -1045,12 +1043,14 @@ namespace RuntimeVirtualTexture
 		View->CachedViewUniformShaderParameters->RuntimeVirtualTexturePackHeight = WorldHeightPackParameter;
 		View->CachedViewUniformShaderParameters->RuntimeVirtualTextureDebugParams = FVector4(DebugType == ERuntimeVirtualTextureDebugType::Debug ? 1.f : 0.f, 0.f, 0.f, 0.f);
 		View->ViewUniformBuffer = TUniformBufferRef<FViewUniformShaderParameters>::CreateUniformBufferImmediate(*View->CachedViewUniformShaderParameters, UniformBuffer_SingleFrame);
-		UploadDynamicPrimitiveShaderDataForView(RHICmdList, *(const_cast<FScene*>(Scene)), *View);
-		Scene->UniformBuffers.VirtualTextureViewUniformBuffer.UpdateUniformBufferImmediate(*View->CachedViewUniformShaderParameters);
+		UploadDynamicPrimitiveShaderDataForView(GraphBuilder.RHICmdList, *(const_cast<FScene*>(Scene)), *View);
+
+		AddPass(GraphBuilder, [Scene, View](FRHICommandList&)
+		{
+			Scene->UniformBuffers.VirtualTextureViewUniformBuffer.UpdateUniformBufferImmediate(*View->CachedViewUniformShaderParameters);
+		});
 
 		// Build graph
-		FMemMark Mark(FMemStack::Get());
-		FRDGBuilder GraphBuilder(RHICmdList);
 		FRenderGraphSetup GraphSetup(GraphBuilder, Scene->GetFeatureLevel(), MaterialType, OutputTexture0, TextureSize, bIsThumbnails);
 
 		// Draw Pass
@@ -1131,26 +1131,23 @@ namespace RuntimeVirtualTexture
 		{
 			AddCopyToOutputPass(GraphBuilder, GraphSetup.OutputAlias2, OutputTexture2, DestBox2);
 		}
-
-		// Execute the graph
-		GraphBuilder.Execute();
 	}
 
-	void RenderPages(FRHICommandListImmediate& RHICmdList, FRenderPageBatchDesc const& InDesc)
+	void RenderPages(FRDGBuilder& GraphBuilder, FRenderPageBatchDesc const& InDesc)
 	{
-		SCOPED_DRAW_EVENT(RHICmdList, RuntimeVirtualTextureRenderPages);
+		RDG_EVENT_SCOPE(GraphBuilder, "RuntimeVirtualTextureRenderPages");
 		check(InDesc.NumPageDescs <= EMaxRenderPageBatch);
 
 		// Make sure GPUScene is up to date. 
 		// Usually this is a no-op since we updated before calling, but RuntimeVirtualTexture::BuildStreamedMips() needs this.
-		UpdateGPUScene(RHICmdList, *InDesc.Scene->GetRenderScene());
+		UpdateGPUScene(GraphBuilder.RHICmdList, *InDesc.Scene->GetRenderScene());
 
 		for (int32 PageIndex = 0; PageIndex < InDesc.NumPageDescs; ++PageIndex)
 		{
 			FRenderPageDesc const& PageDesc = InDesc.PageDescs[PageIndex];
 
 			RenderPage(
-				RHICmdList,
+				GraphBuilder,
 				InDesc.Scene,
 				InDesc.RuntimeVirtualTextureMask,
 				InDesc.MaterialType,
@@ -1168,6 +1165,13 @@ namespace RuntimeVirtualTexture
 		}
 	}
 
+	void RenderPages(FRHICommandListImmediate& RHICmdList, FRenderPageBatchDesc const& InDesc)
+	{
+		FMemMark MemMark(FMemStack::Get());
+		FRDGBuilder GraphBuilder(RHICmdList);
+		RenderPages(GraphBuilder, InDesc);
+		GraphBuilder.Execute();
+	}
 
 	uint32 GetRuntimeVirtualTextureSceneIndex_GameThread(class URuntimeVirtualTextureComponent* InComponent)
 	{
diff --git a/Engine/Source/Runtime/Renderer/Private/VT/VirtualTextureSpace.cpp b/Engine/Source/Runtime/Renderer/Private/VT/VirtualTextureSpace.cpp
--- a/Engine/Source/Runtime/Renderer/Private/VT/VirtualTextureSpace.cpp
+++ b/Engine/Source/Runtime/Renderer/Private/VT/VirtualTextureSpace.cpp
@@ -256,11 +256,11 @@ void FVirtualTextureSpace::QueueUpdateEntirePageTable()
 	bForceEntireUpdate = true;
 }
 
-void FVirtualTextureSpace::AllocateTextures(FRHICommandList& RHICmdList)
+void FVirtualTextureSpace::AllocateTextures(FRDGBuilder& GraphBuilder)
 {
 	if (bNeedToAllocatePageTable)
 	{
-		SCOPED_GPU_MASK(RHICmdList, FRHIGPUMask::All());
+		RDG_GPU_MASK_SCOPE(GraphBuilder, FRHIGPUMask::All());
 
 		const TCHAR* TextureNames[] = { TEXT("PageTable_0"), TEXT("PageTable_1") };
 		static_assert(UE_ARRAY_COUNT(TextureNames) == TextureCapacity, "");
@@ -269,50 +269,32 @@ void FVirtualTextureSpace::AllocateTextures(FRHICommandList& RHICmdList)
 		{
 			// Page Table
 			FTextureEntry& TextureEntry = PageTable[TextureIndex];
-			const FPooledRenderTargetDesc Desc = FPooledRenderTargetDesc::Create2DDesc(
+			const FRDGTextureDesc Desc = FRDGTextureDesc::Create2D(
 				FIntPoint(PageTableSize, PageTableSize),
 				TexturePixelFormat[TextureIndex],
 				FClearValueBinding::None,
-				TexCreate_None,
 				TexCreate_RenderTargetable | TexCreate_ShaderResource,
-				false,
-				NumPageTableLevels,
-				false /* InAutowritable */);
+				NumPageTableLevels);
+
+			FRDGTextureRef DstTexture = GraphBuilder.CreateTexture(Desc, TextureNames[TextureIndex]);
 
-			TRefCountPtr<IPooledRenderTarget> RenderTarget;
-			GRenderTargetPool.FindFreeElement(RHICmdList, Desc, RenderTarget, TextureNames[TextureIndex]);
-			RHIUpdateTextureReference(TextureEntry.TextureReferenceRHI, RenderTarget->GetRenderTargetItem().ShaderResourceTexture);
 			if (TextureEntry.RenderTarget)
 			{
+				FRDGTextureRef SrcTexture = GraphBuilder.RegisterExternalTexture(TextureEntry.RenderTarget);
+				const FRDGTextureDesc& SrcDesc = SrcTexture->Desc;
+
 				// Copy previously allocated page table to new texture
-				const FPooledRenderTargetDesc& SrcDesc = TextureEntry.RenderTarget->GetDesc();
 				FRHICopyTextureInfo CopyInfo;
 				CopyInfo.Size.X = FMath::Min(Desc.Extent.X, SrcDesc.Extent.X);
 				CopyInfo.Size.Y = FMath::Min(Desc.Extent.Y, SrcDesc.Extent.Y);
 				CopyInfo.Size.Z = 1;
 				CopyInfo.NumMips = FMath::Min(Desc.NumMips, SrcDesc.NumMips);
 
-				FRHITexture* SrcTexture = TextureEntry.RenderTarget->GetRenderTargetItem().ShaderResourceTexture;
-				FRHITexture* DstTexture = RenderTarget->GetRenderTargetItem().TargetableTexture;
-
-				FRHITransitionInfo TransitionsBefore[] = {
-					FRHITransitionInfo(SrcTexture, ERHIAccess::Unknown, ERHIAccess::CopySrc),
-					FRHITransitionInfo(DstTexture, ERHIAccess::Unknown, ERHIAccess::CopyDest)
-				};
-				RHICmdList.Transition(MakeArrayView(TransitionsBefore, UE_ARRAY_COUNT(TransitionsBefore)));
-
-				RHICmdList.CopyTexture(SrcTexture, DstTexture, CopyInfo);
-
-				FRHITransitionInfo TransitionsAfter[] = {
-					FRHITransitionInfo(SrcTexture, ERHIAccess::CopySrc, ERHIAccess::SRVGraphics | ERHIAccess::SRVCompute),
-					FRHITransitionInfo(DstTexture, ERHIAccess::CopyDest, ERHIAccess::SRVGraphics | ERHIAccess::SRVCompute)
-				};
-				RHICmdList.Transition(MakeArrayView(TransitionsAfter, UE_ARRAY_COUNT(TransitionsAfter)));
-
-				GRenderTargetPool.FreeUnusedResource(TextureEntry.RenderTarget);
+				AddCopyTexturePass(GraphBuilder, SrcTexture, DstTexture, CopyInfo);
 			}
 
-			TextureEntry.RenderTarget = RenderTarget;
+			ConvertToExternalTexture(GraphBuilder, DstTexture, TextureEntry.RenderTarget);
+			RHIUpdateTextureReference(TextureEntry.TextureReferenceRHI, TextureEntry.RenderTarget->GetShaderResourceRHI());
 		}
 
 		bNeedToAllocatePageTable = false;
@@ -322,17 +304,15 @@ void FVirtualTextureSpace::AllocateTextures(FRHICommandList& RHICmdList)
 	{
 		if (Description.IndirectionTextureSize > 0)
 		{
-			const FPooledRenderTargetDesc Desc = FPooledRenderTargetDesc::Create2DDesc(
+			const FRDGTextureDesc Desc = FRDGTextureDesc::Create2D(
 				FIntPoint(Description.IndirectionTextureSize, Description.IndirectionTextureSize),
 				PF_R32_UINT,
 				FClearValueBinding::None,
-				TexCreate_None,
-				TexCreate_UAV | TexCreate_ShaderResource,
-				false);
+				TexCreate_UAV | TexCreate_ShaderResource);
 
-			GRenderTargetPool.FindFreeElement(RHICmdList, Desc, PageTableIndirection, TEXT("PageTableIndirection"));
-			FRHITexture* TextureRHI = PageTableIndirection->GetRenderTargetItem().ShaderResourceTexture;
-			RHICmdList.ClearUAVUint(RHICreateUnorderedAccessView(TextureRHI), FUintVector4(ForceInitToZero));
+			FRDGTextureRef PageTableIndirectionTexture = GraphBuilder.CreateTexture(Desc, TEXT("PageTableIndirection"));
+			AddClearUAVPass(GraphBuilder, GraphBuilder.CreateUAV(PageTableIndirectionTexture), FUintVector4(ForceInitToZero));
+			ConvertToExternalTexture(GraphBuilder, PageTableIndirectionTexture, PageTableIndirection);
 		}
 
 		bNeedToAllocatePageTableIndirection = false;
@@ -340,7 +320,7 @@ void FVirtualTextureSpace::AllocateTextures(FRHICommandList& RHICmdList)
 }
 
 
-void FVirtualTextureSpace::ApplyUpdates(FVirtualTextureSystem* System, FRHICommandListImmediate& RHICmdList)
+void FVirtualTextureSpace::ApplyUpdates(FVirtualTextureSystem* System, FRDGBuilder& GraphBuilder)
 {
 	static TArray<FPageTableUpdate> ExpandedUpdates[VIRTUALTEXTURE_SPACE_MAXLAYERS][16];
 
@@ -352,7 +332,7 @@ void FVirtualTextureSpace::ApplyUpdates(FVirtualTextureSystem* System, FRHIComma
 	}
 
 	// Multi-GPU support : May be ineffecient for AFR.
-	SCOPED_GPU_MASK(RHICmdList, FRHIGPUMask::All());
+	RDG_GPU_MASK_SCOPE(GraphBuilder, FRHIGPUMask::All());
 
 	for (uint32 LayerIndex = 0u; LayerIndex < Description.NumPageTableLayers; ++LayerIndex)
 	{
@@ -392,10 +372,6 @@ void FVirtualTextureSpace::ApplyUpdates(FVirtualTextureSystem* System, FRHIComma
 
 	if (TotalNumUpdates == 0u)
 	{
-		for (uint32 i = 0u; i < GetNumPageTableTextures(); ++i)
-		{
-			GVisualizeTexture.SetCheckPoint(RHICmdList, PageTable[i].RenderTarget);
-		}
 		return;
 	}
 
@@ -433,16 +409,8 @@ void FVirtualTextureSpace::ApplyUpdates(FVirtualTextureSystem* System, FRHIComma
 		RHIUnlockVertexBuffer(UpdateBuffer);
 	}
 
-	TArray<FRHITransitionInfo, SceneRenderingAllocator> TexturesToTransition;
-	TexturesToTransition.SetNumUninitialized(GetNumPageTableTextures());
-	for (int32 i = 0; i < TexturesToTransition.Num(); ++i)
-	{
-		TexturesToTransition[i] = FRHITransitionInfo(PageTable[i].RenderTarget->GetRenderTargetItem().TargetableTexture, ERHIAccess::Unknown, ERHIAccess::RTV);
-	}
-	RHICmdList.Transition(MakeArrayView(TexturesToTransition.GetData(), TexturesToTransition.Num()));
-
 	// Draw
-	SCOPED_DRAW_EVENT(RHICmdList, PageTableUpdate);
+	RDG_EVENT_SCOPE(GraphBuilder, "PageTableUpdate");
 
 	auto ShaderMap = GetGlobalShaderMap(GetFeatureLevel());
 	TShaderRef<FPageTableUpdateVS> VertexShader;
@@ -463,9 +431,7 @@ void FVirtualTextureSpace::ApplyUpdates(FVirtualTextureSystem* System, FRHIComma
 		const uint32 LayerInTexture = LayerIndex % LayersPerPageTableTexture;
 
 		FTextureEntry& PageTableEntry = PageTable[TextureIndex];
-		check(PageTableEntry.RenderTarget != nullptr);
-
-		FSceneRenderTargetItem& PageTableTarget = PageTableEntry.RenderTarget->GetRenderTargetItem();
+		FRDGTextureRef PageTableTexture = GraphBuilder.RegisterExternalTexture(PageTableEntry.RenderTarget);
 
 		// Use color write mask to update the proper page table entry for this layer
 		FRHIBlendState* BlendStateRHI = nullptr;
@@ -497,40 +463,45 @@ void FVirtualTextureSpace::ApplyUpdates(FVirtualTextureSystem* System, FRHIComma
 			const uint32 NumUpdates = ExpandedUpdates[LayerIndex][Mip].Num();
 			if (NumUpdates)
 			{
-				FRHIRenderPassInfo RPInfo(PageTableTarget.TargetableTexture, ERenderTargetActions::Load_Store, nullptr, Mip);
-				RHICmdList.BeginRenderPass(RPInfo, TEXT("PageTableUpdate"));
-				
-				RHICmdList.SetViewport(0, 0, 0.0f, MipSize, MipSize, 1.0f);
+				auto* PassParameters = GraphBuilder.AllocParameters<FRenderTargetParameters>();
+				PassParameters->RenderTargets[0] = FRenderTargetBinding(PageTableTexture, ERenderTargetLoadAction::ENoAction, Mip);
+
+				GraphBuilder.AddPass(
+					RDG_EVENT_NAME("PageTableUpdate (Mip: %d)", Mip),
+					PassParameters,
+					ERDGPassFlags::Raster,
+					[this, VertexShader, PixelShader, BlendStateRHI, FirstUpdate, NumUpdates, MipSize](FRHICommandList& RHICmdList)
+				{
+					RHICmdList.SetViewport(0, 0, 0.0f, MipSize, MipSize, 1.0f);
 
-				FGraphicsPipelineStateInitializer GraphicsPSOInit;
-				RHICmdList.ApplyCachedRenderTargets(GraphicsPSOInit);
+					FGraphicsPipelineStateInitializer GraphicsPSOInit;
+					RHICmdList.ApplyCachedRenderTargets(GraphicsPSOInit);
 
-				GraphicsPSOInit.BlendState = BlendStateRHI;
-				GraphicsPSOInit.RasterizerState = TStaticRasterizerState<>::GetRHI();
-				GraphicsPSOInit.DepthStencilState = TStaticDepthStencilState<false, CF_Always>::GetRHI();
-				GraphicsPSOInit.PrimitiveType = PT_TriangleList;
+					GraphicsPSOInit.BlendState = BlendStateRHI;
+					GraphicsPSOInit.RasterizerState = TStaticRasterizerState<>::GetRHI();
+					GraphicsPSOInit.DepthStencilState = TStaticDepthStencilState<false, CF_Always>::GetRHI();
+					GraphicsPSOInit.PrimitiveType = PT_TriangleList;
 
-				GraphicsPSOInit.BoundShaderState.VertexDeclarationRHI = GEmptyVertexDeclaration.VertexDeclarationRHI;
-				GraphicsPSOInit.BoundShaderState.VertexShaderRHI = VertexShader.GetVertexShader();
-				GraphicsPSOInit.BoundShaderState.PixelShaderRHI = PixelShader.GetPixelShader();
+					GraphicsPSOInit.BoundShaderState.VertexDeclarationRHI = GEmptyVertexDeclaration.VertexDeclarationRHI;
+					GraphicsPSOInit.BoundShaderState.VertexShaderRHI = VertexShader.GetVertexShader();
+					GraphicsPSOInit.BoundShaderState.PixelShaderRHI = PixelShader.GetPixelShader();
 
-				SetGraphicsPipelineState(RHICmdList, GraphicsPSOInit);
+					SetGraphicsPipelineState(RHICmdList, GraphicsPSOInit);
 
-				{
-					FRHIVertexShader* ShaderRHI = VertexShader.GetVertexShader();
-					SetShaderValue(RHICmdList, ShaderRHI, VertexShader->PageTableSize, PageTableSize);
-					SetShaderValue(RHICmdList, ShaderRHI, VertexShader->FirstUpdate, FirstUpdate);
-					SetShaderValue(RHICmdList, ShaderRHI, VertexShader->NumUpdates, NumUpdates);
-					SetSRVParameter(RHICmdList, ShaderRHI, VertexShader->UpdateBuffer, UpdateBufferSRV);
-				}
+					{
+						FRHIVertexShader* ShaderRHI = VertexShader.GetVertexShader();
+						SetShaderValue(RHICmdList, ShaderRHI, VertexShader->PageTableSize, PageTableSize);
+						SetShaderValue(RHICmdList, ShaderRHI, VertexShader->FirstUpdate, FirstUpdate);
+						SetShaderValue(RHICmdList, ShaderRHI, VertexShader->NumUpdates, NumUpdates);
+						SetSRVParameter(RHICmdList, ShaderRHI, VertexShader->UpdateBuffer, UpdateBufferSRV);
+					}
 
-				// needs to be the same on shader side (faster on NVIDIA and AMD)
-				uint32 QuadsPerInstance = 8;
+					// needs to be the same on shader side (faster on NVIDIA and AMD)
+					uint32 QuadsPerInstance = 8;
 
-				RHICmdList.SetStreamSource(0, NULL, 0);
-				RHICmdList.DrawIndexedPrimitive(GQuadIndexBuffer.IndexBufferRHI, 0, 0, 32, 0, 2 * QuadsPerInstance, FMath::DivideAndRoundUp(NumUpdates, QuadsPerInstance));
-
-				RHICmdList.EndRenderPass();
+					RHICmdList.SetStreamSource(0, NULL, 0);
+					RHICmdList.DrawIndexedPrimitive(GQuadIndexBuffer.IndexBufferRHI, 0, 0, 32, 0, 2 * QuadsPerInstance, FMath::DivideAndRoundUp(NumUpdates, QuadsPerInstance));
+				});
 
 				ExpandedUpdates[LayerIndex][Mip].Reset();
 			}
@@ -539,13 +510,6 @@ void FVirtualTextureSpace::ApplyUpdates(FVirtualTextureSystem* System, FRHIComma
 			MipSize >>= 1;
 		}
 	}
-
-	for (uint32 i = 0u; i < GetNumPageTableTextures(); ++i)
-	{
-		FSceneRenderTargetItem& PageTableTarget = PageTable[i].RenderTarget->GetRenderTargetItem();
-		RHICmdList.CopyToResolveTarget(PageTableTarget.TargetableTexture, PageTableTarget.ShaderResourceTexture, FResolveParams());
-		GVisualizeTexture.SetCheckPoint(RHICmdList, PageTable[i].RenderTarget);
-	}
 }
 
 void FVirtualTextureSpace::DumpToConsole(bool verbose)
diff --git a/Engine/Source/Runtime/Renderer/Private/VT/VirtualTextureSpace.h b/Engine/Source/Runtime/Renderer/Private/VT/VirtualTextureSpace.h
--- a/Engine/Source/Runtime/Renderer/Private/VT/VirtualTextureSpace.h
+++ b/Engine/Source/Runtime/Renderer/Private/VT/VirtualTextureSpace.h
@@ -87,8 +87,8 @@ public:
 	FRHITexture* GetPageTableIndirectionTexture() const;
 
 	void				QueueUpdate( uint8 Layer, uint8 vLogSize, uint32 vAddress, uint8 vLevel, const FPhysicalTileLocation& pTileLocation);
-	void				AllocateTextures(FRHICommandList& RHICmdList);
-	void				ApplyUpdates(FVirtualTextureSystem* System, FRHICommandListImmediate& RHICmdList);
+	void				AllocateTextures(FRDGBuilder& GraphBuilder);
+	void				ApplyUpdates(FVirtualTextureSystem* System, FRDGBuilder& GraphBuilder);
 	void				QueueUpdateEntirePageTable();
 
 	void DumpToConsole(bool verbose);
diff --git a/Engine/Source/Runtime/Renderer/Private/VT/VirtualTextureSystem.cpp b/Engine/Source/Runtime/Renderer/Private/VT/VirtualTextureSystem.cpp
--- a/Engine/Source/Runtime/Renderer/Private/VT/VirtualTextureSystem.cpp
+++ b/Engine/Source/Runtime/Renderer/Private/VT/VirtualTextureSystem.cpp
@@ -850,7 +850,7 @@ void FVirtualTextureSystem::RequestTilesForRegion(const IAllocatedVirtualTexture
 	}
 }
 
-void FVirtualTextureSystem::LoadPendingTiles(FRHICommandListImmediate& RHICmdList, ERHIFeatureLevel::Type FeatureLevel)
+void FVirtualTextureSystem::LoadPendingTiles(FRDGBuilder& GraphBuilder, ERHIFeatureLevel::Type FeatureLevel)
 {
 	check(IsInRenderingThread());
 
@@ -865,8 +865,6 @@ void FVirtualTextureSystem::LoadPendingTiles(FRHICommandListImmediate& RHICmdLis
 	if (PackedTiles.Num() > 0)
 	{
 		FMemStack& MemStack = FMemStack::Get();
-		FMemMark Mark(MemStack);
-
 		FUniquePageList* UniquePageList = new(MemStack) FUniquePageList;
 		UniquePageList->Initialize();
 		for (uint32 Tile : PackedTiles)
@@ -878,8 +876,8 @@ void FVirtualTextureSystem::LoadPendingTiles(FRHICommandListImmediate& RHICmdLis
 		RequestList->Initialize();
 		GatherRequests(RequestList, UniquePageList, Frame, MemStack);
 		// No need to sort requests, since we're submitting all of them here (no throttling)
-		AllocateResources(RHICmdList, FeatureLevel);
-		SubmitRequests(RHICmdList, FeatureLevel, MemStack, RequestList, false);
+		AllocateResources(GraphBuilder, FeatureLevel);
+		SubmitRequests(GraphBuilder, FeatureLevel, MemStack, RequestList, false);
 	}
 }
 
@@ -938,14 +936,14 @@ void FVirtualTextureSystem::FeedbackAnalysisTask(const FFeedbackAnalysisParamete
 	}
 }
 
-void FVirtualTextureSystem::Update(FRHICommandListImmediate& RHICmdList, ERHIFeatureLevel::Type FeatureLevel, FScene* Scene)
+void FVirtualTextureSystem::Update(FRDGBuilder& GraphBuilder, ERHIFeatureLevel::Type FeatureLevel, FScene* Scene)
 {
 	check(IsInRenderingThread());
 
 	CSV_SCOPED_TIMING_STAT_EXCLUSIVE(VirtualTextureSystem_Update);
 	SCOPE_CYCLE_COUNTER(STAT_VirtualTextureSystem_Update);
-	SCOPED_GPU_STAT(RHICmdList, VirtualTexture);
-	
+	RDG_GPU_STAT_SCOPE(GraphBuilder, VirtualTexture);
+
 	if (bFlushCaches)
 	{
 		SCOPE_CYCLE_COUNTER(STAT_FlushCache);
@@ -983,7 +981,7 @@ void FVirtualTextureSystem::Update(FRHICommandListImmediate& RHICmdList, ERHIFea
 	{
 		// Only flush if we know that there is GPU feedback available to refill the visible data this frame
 		// This prevents bugs when low frame rate causes feedback buffer to stall so that the physical cache isn't filled immediately which causes visible glitching
-		if (GVirtualTextureFeedback.CanMap(RHICmdList))
+		if (GVirtualTextureFeedback.CanMap(GraphBuilder.RHICmdList))
 		{
 			// Each RVT will call FVirtualTextureSystem::FlushCache()
 			Scene->FlushDirtyRuntimeVirtualTextures();
@@ -1003,7 +1001,6 @@ void FVirtualTextureSystem::Update(FRHICommandListImmediate& RHICmdList, ERHIFea
 	}
 
 	FMemStack& MemStack = FMemStack::Get();
-	FMemMark Mark(MemStack);
 	FUniquePageList* MergedUniquePageList = new(MemStack) FUniquePageList;
 	MergedUniquePageList->Initialize();
 	
@@ -1016,7 +1013,7 @@ void FVirtualTextureSystem::Update(FRHICommandListImmediate& RHICmdList, ERHIFea
 
 		{
 			SCOPE_CYCLE_COUNTER(STAT_FeedbackMap);
-			FeedbackResult = GVirtualTextureFeedback.Map(RHICmdList);
+			FeedbackResult = GVirtualTextureFeedback.Map(GraphBuilder.RHICmdList);
 		}
 
 		// Create tasks to read the feedback data
@@ -1089,7 +1086,7 @@ void FVirtualTextureSystem::Update(FRHICommandListImmediate& RHICmdList, ERHIFea
 			}
 		}
 
-		GVirtualTextureFeedback.Unmap(RHICmdList, FeedbackResult.MapHandle);
+		GVirtualTextureFeedback.Unmap(GraphBuilder.RHICmdList, FeedbackResult.MapHandle);
 	}
 
 	FUniqueRequestList* MergedRequestList = new(MemStack) FUniqueRequestList(MemStack);
@@ -1142,7 +1139,6 @@ void FVirtualTextureSystem::Update(FRHICommandListImmediate& RHICmdList, ERHIFea
 	{
 		// Collect explicitly requested tiles
 		// These tiles are generated on the current frame, so they are collected/processed in a separate list
-		FMemMark RequestPageMark(MemStack);
 		FUniquePageList* RequestedPageList = new(MemStack) FUniquePageList;
 		RequestedPageList->Initialize();
 		for (uint32 Tile : PackedTiles)
@@ -1190,9 +1186,9 @@ void FVirtualTextureSystem::Update(FRHICommandListImmediate& RHICmdList, ERHIFea
 	}
 
 	// Submit the requests to produce pages that are already mapped
-	SubmitPreMappedRequests(RHICmdList, FeatureLevel);
+	SubmitPreMappedRequests(GraphBuilder, FeatureLevel);
 	// Submit the merged requests
-	SubmitRequests(RHICmdList, FeatureLevel, MemStack, MergedRequestList, true);
+	SubmitRequests(GraphBuilder, FeatureLevel, MemStack, MergedRequestList, true);
 
 	UpdateCSVStats();
 
@@ -1767,7 +1763,7 @@ void FVirtualTextureSystem::UpdateCSVStats() const
 #endif
 }
 
-void FVirtualTextureSystem::SubmitRequestsFromLocalTileList(const TSet<FVirtualTextureLocalTile>& LocalTileList, EVTProducePageFlags Flags, FRHICommandListImmediate& RHICmdList, ERHIFeatureLevel::Type FeatureLevel)
+void FVirtualTextureSystem::SubmitRequestsFromLocalTileList(const TSet<FVirtualTextureLocalTile>& LocalTileList, EVTProducePageFlags Flags, FRDGBuilder& GraphBuilder, ERHIFeatureLevel::Type FeatureLevel)
 {
 	LLM_SCOPE(ELLMTag::VirtualTextureSystem);
 
@@ -1820,7 +1816,7 @@ void FVirtualTextureSystem::SubmitRequestsFromLocalTileList(const TSet<FVirtualT
 		}
 
 		IVirtualTextureFinalizer* VTFinalizer = Producer.GetVirtualTexture()->ProducePageData(
-			RHICmdList, FeatureLevel,
+			GraphBuilder.RHICmdList, FeatureLevel,
 			Flags,
 			ProducerHandle, LayerMask, Tile.Local_vLevel, Tile.Local_vAddress,
 			RequestPageResult.Handle,
@@ -1834,27 +1830,25 @@ void FVirtualTextureSystem::SubmitRequestsFromLocalTileList(const TSet<FVirtualT
 	}
 }
 
-void FVirtualTextureSystem::SubmitPreMappedRequests(FRHICommandListImmediate& RHICmdList, ERHIFeatureLevel::Type FeatureLevel)
+void FVirtualTextureSystem::SubmitPreMappedRequests(FRDGBuilder& GraphBuilder, ERHIFeatureLevel::Type FeatureLevel)
 {
 	{
 		INC_DWORD_STAT_BY(STAT_NumMappedPageUpdate, MappedTilesToProduce.Num());
-		SubmitRequestsFromLocalTileList(MappedTilesToProduce, EVTProducePageFlags::None, RHICmdList, FeatureLevel);
+		SubmitRequestsFromLocalTileList(MappedTilesToProduce, EVTProducePageFlags::None, GraphBuilder, FeatureLevel);
 		MappedTilesToProduce.Reset();
 	}
 
 	{
 		INC_DWORD_STAT_BY(STAT_NumContinuousPageUpdate, ContinuousUpdateTilesToProduce.Num());
-		SubmitRequestsFromLocalTileList(ContinuousUpdateTilesToProduce, EVTProducePageFlags::ContinuousUpdate, RHICmdList, FeatureLevel);
+		SubmitRequestsFromLocalTileList(ContinuousUpdateTilesToProduce, EVTProducePageFlags::ContinuousUpdate, GraphBuilder, FeatureLevel);
 		ContinuousUpdateTilesToProduce.Reset();
 	}
 }
 
-void FVirtualTextureSystem::SubmitRequests(FRHICommandListImmediate& RHICmdList, ERHIFeatureLevel::Type FeatureLevel, FMemStack& MemStack, FUniqueRequestList* RequestList, bool bAsync)
+void FVirtualTextureSystem::SubmitRequests(FRDGBuilder& GraphBuilder, ERHIFeatureLevel::Type FeatureLevel, FMemStack& MemStack, FUniqueRequestList* RequestList, bool bAsync)
 {
 	LLM_SCOPE(ELLMTag::VirtualTextureSystem);
 
-	FMemMark Mark(MemStack);
-
 	// Allocate space to hold the physical address we allocate for each page load (1 page per layer per request)
 	uint32* RequestPhysicalAddress = new(MemStack, MEM_Oned) uint32[RequestList->GetNumLoadRequests() * VIRTUALTEXTURE_SPACE_MAXLAYERS];
 	{
@@ -2072,7 +2066,7 @@ void FVirtualTextureSystem::SubmitRequests(FRHICommandListImmediate& RHICmdList,
 
 			for (FProducePageDataPrepareTask& Task : PrepareTasks)
 			{
-				IVirtualTextureFinalizer* VTFinalizer = Task.VirtualTexture->ProducePageData(RHICmdList, FeatureLevel,
+				IVirtualTextureFinalizer* VTFinalizer = Task.VirtualTexture->ProducePageData(GraphBuilder.RHICmdList, FeatureLevel,
 					Task.Flags,
 					Task.ProducerHandle, Task.LayerMask, Task.vLevel, Task.vAddress,
 					Task.RequestHandle,
@@ -2218,7 +2212,7 @@ void FVirtualTextureSystem::SubmitRequests(FRHICommandListImmediate& RHICmdList,
 		SCOPE_CYCLE_COUNTER(STAT_ProcessRequests_Finalize);
 		for (IVirtualTextureFinalizer* VTFinalizer : Finalizers)
 		{
-			VTFinalizer->Finalize(RHICmdList);
+			VTFinalizer->Finalize(GraphBuilder);
 		}
 		Finalizers.Reset();
 	}
@@ -2230,7 +2224,7 @@ void FVirtualTextureSystem::SubmitRequests(FRHICommandListImmediate& RHICmdList,
 		{
 			if (Spaces[ID])
 			{
-				Spaces[ID]->ApplyUpdates(this, RHICmdList);
+				Spaces[ID]->ApplyUpdates(this, GraphBuilder);
 			}
 		}
 	}
@@ -2238,16 +2232,16 @@ void FVirtualTextureSystem::SubmitRequests(FRHICommandListImmediate& RHICmdList,
 	Frame++;
 }
 
-void FVirtualTextureSystem::AllocateResources(FRHICommandListImmediate& RHICmdList, ERHIFeatureLevel::Type FeatureLevel)
+void FVirtualTextureSystem::AllocateResources(FRDGBuilder& GraphBuilder, ERHIFeatureLevel::Type FeatureLevel)
 {
 	LLM_SCOPE(ELLMTag::VirtualTextureSystem);
-	SCOPED_GPU_STAT(RHICmdList, VirtualTextureAllocate);
+	RDG_GPU_STAT_SCOPE(GraphBuilder, VirtualTextureAllocate);
 
 	for (uint32 ID = 0; ID < MaxSpaces; ID++)
 	{
 		if (Spaces[ID])
 		{
-			Spaces[ID]->AllocateTextures(RHICmdList);
+			Spaces[ID]->AllocateTextures(GraphBuilder);
 		}
 	}
 }
diff --git a/Engine/Source/Runtime/Renderer/Private/VT/VirtualTextureSystem.h b/Engine/Source/Runtime/Renderer/Private/VT/VirtualTextureSystem.h
--- a/Engine/Source/Runtime/Renderer/Private/VT/VirtualTextureSystem.h
+++ b/Engine/Source/Runtime/Renderer/Private/VT/VirtualTextureSystem.h
@@ -34,9 +34,9 @@ public:
 	static void Shutdown();
 	static FVirtualTextureSystem& Get();
 
-	void AllocateResources(FRHICommandListImmediate& RHICmdList, ERHIFeatureLevel::Type FeatureLevel);
+	void AllocateResources(FRDGBuilder& GraphBuilder, ERHIFeatureLevel::Type FeatureLevel);
 	void CallPendingCallbacks();
-	void Update( FRHICommandListImmediate& RHICmdList, ERHIFeatureLevel::Type FeatureLevel, FScene* Scene);
+	void Update( FRDGBuilder& GraphBuilder, ERHIFeatureLevel::Type FeatureLevel, FScene* Scene);
 	void ReleasePendingResources();
 
 	IAllocatedVirtualTexture* AllocateVirtualTexture(const FAllocatedVTDescription& Desc);
@@ -67,7 +67,7 @@ public:
 	void ForceUnlockAllTiles(const FVirtualTextureProducerHandle& ProducerHandle, const FVirtualTextureProducer* Producer);
 	void RequestTiles(const FVector2D& InScreenSpaceSize, int32 InMipLevel = -1);
 	void RequestTilesForRegion(const IAllocatedVirtualTexture* AllocatedVT, const FVector2D& InScreenSpaceSize, const FIntRect& InTextureRegion, int32 InMipLevel = -1);
-	void LoadPendingTiles(FRHICommandListImmediate& RHICmdList, ERHIFeatureLevel::Type FeatureLevel);
+	void LoadPendingTiles(FRDGBuilder& GraphBuilder, ERHIFeatureLevel::Type FeatureLevel);
 	
 	void FlushCache();
 	void FlushCache(FVirtualTextureProducerHandle const& ProducerHandle, FIntRect const& TextureRegion, uint32 MaxLevel);
@@ -84,11 +84,11 @@ private:
 
 	void RequestTilesForRegionInternal(const IAllocatedVirtualTexture* AllocatedVT, const FIntRect& InTextureRegion, uint32 vLevel);
 	
-	void SubmitRequestsFromLocalTileList(const TSet<FVirtualTextureLocalTile>& LocalTileList, EVTProducePageFlags Flags, FRHICommandListImmediate& RHICmdList, ERHIFeatureLevel::Type FeatureLevel);
+	void SubmitRequestsFromLocalTileList(const TSet<FVirtualTextureLocalTile>& LocalTileList, EVTProducePageFlags Flags, FRDGBuilder& GraphBuilder, ERHIFeatureLevel::Type FeatureLevel);
 
-	void SubmitPreMappedRequests(FRHICommandListImmediate& RHICmdList, ERHIFeatureLevel::Type FeatureLevel);
+	void SubmitPreMappedRequests(FRDGBuilder& GraphBuilder, ERHIFeatureLevel::Type FeatureLevel);
 
-	void SubmitRequests(FRHICommandListImmediate& RHICmdList, ERHIFeatureLevel::Type FeatureLevel, FMemStack& MemStack, FUniqueRequestList* RequestList, bool bAsync);
+	void SubmitRequests(FRDGBuilder& GraphBuilder, ERHIFeatureLevel::Type FeatureLevel, FMemStack& MemStack, FUniqueRequestList* RequestList, bool bAsync);
 
 	void GatherRequests(FUniqueRequestList* MergedRequestList, const FUniquePageList* UniquePageList, uint32 FrameRequested, FMemStack& MemStack);
 
diff --git a/Engine/Source/Runtime/Renderer/Public/RuntimeVirtualTextureProducer.h b/Engine/Source/Runtime/Renderer/Public/RuntimeVirtualTextureProducer.h
--- a/Engine/Source/Runtime/Renderer/Public/RuntimeVirtualTextureProducer.h
+++ b/Engine/Source/Runtime/Renderer/Public/RuntimeVirtualTextureProducer.h
@@ -37,7 +37,7 @@ public:
 	void AddTile(FTileEntry& Tile);
 
 	//~ Begin IVirtualTextureFinalizer Interface.
-	virtual void Finalize(FRHICommandListImmediate& RHICmdList) override;
+	virtual void Finalize(FRDGBuilder& GraphBuilder) override;
 	//~ End IVirtualTextureFinalizer Interface.
 
 private:
diff --git a/Engine/Source/Runtime/Renderer/Public/VT/RuntimeVirtualTextureRender.h b/Engine/Source/Runtime/Renderer/Public/VT/RuntimeVirtualTextureRender.h
--- a/Engine/Source/Runtime/Renderer/Public/VT/RuntimeVirtualTextureRender.h
+++ b/Engine/Source/Runtime/Renderer/Public/VT/RuntimeVirtualTextureRender.h
@@ -83,5 +83,8 @@ namespace RuntimeVirtualTexture
 	RENDERER_API bool IsSceneReadyToRender(FScene* Scene);
 
 	/** Render a batch of pages for a runtime virtual texture. */
+	RENDERER_API void RenderPages(FRDGBuilder& GraphBuilder, FRenderPageBatchDesc const& InDesc);
+
+	UE_DEPRECATED(5.0, "This method has been refactored to use an FRDGBuilder instead.")
 	RENDERER_API void RenderPages(FRHICommandListImmediate& RHICmdList, FRenderPageBatchDesc const& InDesc);
 }
