CIS fix.

diff --git a/Engine/Plugins/Experimental/GPULightmass/Source/GPULightmass/Private/LightmapRenderer.cpp b/Engine/Plugins/Experimental/GPULightmass/Source/GPULightmass/Private/LightmapRenderer.cpp
--- a/Engine/Plugins/Experimental/GPULightmass/Source/GPULightmass/Private/LightmapRenderer.cpp
+++ b/Engine/Plugins/Experimental/GPULightmass/Source/GPULightmass/Private/LightmapRenderer.cpp
@@ -680,7 +680,7 @@ void FSceneRenderState::SetupRayTracingScene()
 			// Expanded version of View.InitRHIResources() - need to do SetupSkyIrradianceEnvironmentMapConstants manually because the estimation of skylight is dependent on GetSkySHDiffuse
 			View.CachedViewUniformShaderParameters = MakeUnique<FViewUniformShaderParameters>();
 
-			FSceneRenderTargets& SceneContext = FSceneRenderTargets::Get(FRHICommandListExecutor::GetImmediateCommandList());
+			FSceneRenderTargets& SceneContext = FSceneRenderTargets::Get();
 
 			FBox UnusedVolumeBounds[TVC_MAX];
 			View.SetupUniformBufferParameters(
