GPULM: Fix compatibility with FGPUScenePrimitiveCollector changes

diff --git a/Engine/Plugins/Experimental/GPULightmass/Source/GPULightmass/Private/LightmapRenderer.cpp b/Engine/Plugins/Experimental/GPULightmass/Source/GPULightmass/Private/LightmapRenderer.cpp
--- a/Engine/Plugins/Experimental/GPULightmass/Source/GPULightmass/Private/LightmapRenderer.cpp
+++ b/Engine/Plugins/Experimental/GPULightmass/Source/GPULightmass/Private/LightmapRenderer.cpp
@@ -705,6 +705,11 @@ void FSceneRenderState::SetupRayTracingScene()
 
 			CachedRayTracingScene->CachedViewUniformBuffer.UpdateUniformBufferImmediate(*View.CachedViewUniformShaderParameters);
 		}
+
+		const_cast<TRange<int32>&>(View.DynamicPrimitiveCollector.GetPrimitiveIdRange()) = TRange<int32>(0,
+			FMath::Max(StaticMeshInstanceRenderStates.Elements.Num(), FMath::Max(InstanceGroupRenderStates.Elements.Num(), LandscapeRenderStates.Elements.Num()))
+			);
+		View.DynamicPrimitiveCollector.Commit();
 	}
 
 #if RHI_RAYTRACING
diff --git a/Engine/Source/Runtime/Renderer/Private/GPUScene.h b/Engine/Source/Runtime/Renderer/Private/GPUScene.h
--- a/Engine/Source/Runtime/Renderer/Private/GPUScene.h
+++ b/Engine/Source/Runtime/Renderer/Private/GPUScene.h
@@ -70,7 +70,7 @@ public:
 	 * After this is called no more calls to Add are allowed.
 	 * Only allowed inside a FGPUScene::Begin/EndRender block.
 	 */
-	void Commit();
+	RENDERER_API void Commit();
 
 	/**
 	 */
