Change from SOA to AOS layout for path tracing lights

This was recommended by Nvidia in the RT path and also makes the code a bit cleaner. It also reduces the number of places that impose an upper limit on light count.

diff --git a/Engine/Plugins/Experimental/GPULightmass/Shaders/Private/LightmapPathTracing.usf b/Engine/Plugins/Experimental/GPULightmass/Shaders/Private/LightmapPathTracing.usf
--- a/Engine/Plugins/Experimental/GPULightmass/Shaders/Private/LightmapPathTracing.usf
+++ b/Engine/Plugins/Experimental/GPULightmass/Shaders/Private/LightmapPathTracing.usf
@@ -249,7 +249,7 @@ void PathTracingKernel(
 			// Loop over lights to capture their contribution
 			// #dxr_todo: if we have lots of lights, having some hierarchical structure would be better ....
 			float PreviousCdfValue = 0.0;
-			for (uint LightId = 0; LightId < SceneLightsData.Count; ++LightId)
+			for (uint LightId = 0; LightId < SceneLightCount; ++LightId)
 			{
 				// Camera rays should only do the rest of the loop if we have visible lights OR if
 				// we have a skydome (which are always visible)
@@ -369,7 +369,7 @@ void PathTracingKernel(
 		bool bIsNeeValid = IsNeeValidMaterial(Payload);
 		
 		// If we are using Light sampling and we are not on the last bounce ....
-		if (MISMode != 0 && Bounce < 32 && bIsNeeValid && SceneLightsData.Count > 0)
+		if (MISMode != 0 && Bounce < 32 && bIsNeeValid && SceneLightCount > 0)
 		{
 			// Choose a light and sample it
 
diff --git a/Engine/Plugins/Experimental/GPULightmass/Source/GPULightmass/Private/LightmapRayTracing.h b/Engine/Plugins/Experimental/GPULightmass/Source/GPULightmass/Private/LightmapRayTracing.h
--- a/Engine/Plugins/Experimental/GPULightmass/Source/GPULightmass/Private/LightmapRayTracing.h
+++ b/Engine/Plugins/Experimental/GPULightmass/Source/GPULightmass/Private/LightmapRayTracing.h
@@ -10,6 +10,7 @@
 #include "SceneView.h"
 #include "RayTracing/RayTracingMaterialHitShaders.h"
 #include "IrradianceCaching.h"
+#include "RayTracingTypes.h"
 
 #if RHI_RAYTRACING
 
@@ -68,7 +69,8 @@ class FLightmapPathTracingRGS : public FGlobalShader
 		SHADER_PARAMETER_RDG_TEXTURE(Texture2D<float>, RayGuidingCDFX)
 		SHADER_PARAMETER_RDG_TEXTURE(Texture2D<float>, RayGuidingCDFY)
 		SHADER_PARAMETER_SRV(StructuredBuffer<FGPUTileDescription>, BatchedTiles)
-		SHADER_PARAMETER_STRUCT_REF(FPathTracingLightData, LightParameters)
+		SHADER_PARAMETER_RDG_BUFFER_SRV(StructuredBuffer<FPathTracingLight>, SceneLights)
+		SHADER_PARAMETER(uint32, SceneLightCount)
 		SHADER_PARAMETER_STRUCT_REF(FSkyLightData, SkyLight)
 		SHADER_PARAMETER_STRUCT_REF(FViewUniformShaderParameters, ViewUniformBuffer)
 		SHADER_PARAMETER_STRUCT_REF(FIrradianceCachingParameters, IrradianceCachingParameters)
@@ -120,7 +122,8 @@ class FVolumetricLightmapPathTracingRGS : public FGlobalShader
 		SHADER_PARAMETER_UAV(RWTexture3D<float4>, OutSHCoefficients1G)
 		SHADER_PARAMETER_UAV(RWTexture3D<float4>, OutSHCoefficients0B)
 		SHADER_PARAMETER_UAV(RWTexture3D<float4>, OutSHCoefficients1B)
-		SHADER_PARAMETER_STRUCT_REF(FPathTracingLightData, LightParameters)
+		SHADER_PARAMETER_RDG_BUFFER_SRV(StructuredBuffer<FPathTracingLight>, SceneLights)
+		SHADER_PARAMETER(uint32, SceneLightCount)
 		SHADER_PARAMETER_STRUCT_REF(FSkyLightData, SkyLight)
 		SHADER_PARAMETER_STRUCT_REF(FViewUniformShaderParameters, ViewUniformBuffer)
 		SHADER_PARAMETER_STRUCT_REF(FIrradianceCachingParameters, IrradianceCachingParameters)
diff --git a/Engine/Plugins/Experimental/GPULightmass/Source/GPULightmass/Private/LightmapRenderer.cpp b/Engine/Plugins/Experimental/GPULightmass/Source/GPULightmass/Private/LightmapRenderer.cpp
--- a/Engine/Plugins/Experimental/GPULightmass/Source/GPULightmass/Private/LightmapRenderer.cpp
+++ b/Engine/Plugins/Experimental/GPULightmass/Source/GPULightmass/Private/LightmapRenderer.cpp
@@ -165,119 +165,121 @@ struct FGPUBatchedTileRequests
 };
 
 #if RHI_RAYTRACING
-FPathTracingLightData SetupPathTracingLightParameters(const GPULightmass::FLightSceneRenderState& LightScene)
+void SetupPathTracingLightParameters(const GPULightmass::FLightSceneRenderState& LightScene, FRDGBuilder& GraphBuilder, FRDGBufferSRV** OutLightBuffer, unsigned* OutLightCount)
 {
-	FPathTracingLightData LightParameters;
+	const uint32 MaxLightCount = RAY_TRACING_LIGHT_COUNT_MAXIMUM;
+	FPathTracingLight Lights[RAY_TRACING_LIGHT_COUNT_MAXIMUM] = {};
 
-	FMemory::Memzero(LightParameters);
+	unsigned LightCount = 0;
 
-	LightParameters.Count = 0;
-
-	// Prepend SkyLight to light buffer
-	// WARNING: Until ray payload encodes Light data buffer, the execution depends on this ordering!
-	uint32 SkyLightIndex = 0;
-	LightParameters.Color[SkyLightIndex] = LightScene.SkyLight.IsSet() ? FVector(LightScene.SkyLight->Color) : FVector(0);
-	LightParameters.Flags[SkyLightIndex] = PATHTRACER_FLAG_TRANSMISSION_MASK;
-	LightParameters.Flags[SkyLightIndex] |= PATHTRACER_FLAG_LIGHTING_CHANNEL_MASK;
-	bool SkyLightIsStationary = LightScene.SkyLight.IsSet() && LightScene.SkyLight->bStationary;
-	LightParameters.Flags[SkyLightIndex] |= SkyLightIsStationary ? PATHTRACER_FLAG_STATIONARY_MASK : 0;
-	LightParameters.Flags[LightParameters.Count] |= PATHTRACING_LIGHT_SKY;
-	LightParameters.Count++;
-
-	uint32 MaxLightCount = RAY_TRACING_LIGHT_COUNT_MAXIMUM;
+	if (LightScene.SkyLight.IsSet())
+	{
+		FPathTracingLight& DestLight = Lights[LightCount++];
+		DestLight.Color = FVector(LightScene.SkyLight->Color);
+		DestLight.Flags = PATHTRACER_FLAG_TRANSMISSION_MASK;
+		DestLight.Flags |= PATHTRACER_FLAG_LIGHTING_CHANNEL_MASK;
+		bool SkyLightIsStationary = LightScene.SkyLight->bStationary;
+		DestLight.Flags |= SkyLightIsStationary ? PATHTRACER_FLAG_STATIONARY_MASK : 0;
+		DestLight.Flags |= PATHTRACING_LIGHT_SKY;
+	}
 
 	for (auto Light : LightScene.DirectionalLights.Elements)
 	{
-		if (LightParameters.Count < MaxLightCount)
+		if (LightCount < MaxLightCount)
 		{
-			LightParameters.Normal[LightParameters.Count] = -Light.Direction;
-			LightParameters.Color[LightParameters.Count] = FVector(Light.Color);
-			LightParameters.Dimensions[LightParameters.Count] = FVector(0.0f, 0.0f, FMath::Sin(0.5f * FMath::DegreesToRadians(Light.LightSourceAngle)));
-			LightParameters.Attenuation[LightParameters.Count] = 1.0;
-			LightParameters.IESTextureSlice[LightParameters.Count] = -1;
-
-			LightParameters.Flags[LightParameters.Count] = PATHTRACER_FLAG_TRANSMISSION_MASK;
-			LightParameters.Flags[LightParameters.Count] |= PATHTRACER_FLAG_LIGHTING_CHANNEL_MASK;
-			LightParameters.Flags[LightParameters.Count] |= Light.bStationary ? PATHTRACER_FLAG_STATIONARY_MASK : 0;
-			LightParameters.Flags[LightParameters.Count] |= PATHTRACING_LIGHT_DIRECTIONAL;
-
-			LightParameters.Count++;
+			FPathTracingLight& DestLight = Lights[LightCount++];
+
+			DestLight.Normal = -Light.Direction;
+			DestLight.Color = FVector(Light.Color);
+			DestLight.Dimensions = FVector(0.0f, 0.0f, FMath::Sin(0.5f * FMath::DegreesToRadians(Light.LightSourceAngle)));
+			DestLight.Attenuation = 1.0;
+			DestLight.IESTextureSlice = -1;
+
+			DestLight.Flags = PATHTRACER_FLAG_TRANSMISSION_MASK;
+			DestLight.Flags |= PATHTRACER_FLAG_LIGHTING_CHANNEL_MASK;
+			DestLight.Flags |= Light.bStationary ? PATHTRACER_FLAG_STATIONARY_MASK : 0;
+			DestLight.Flags |= PATHTRACING_LIGHT_DIRECTIONAL;
 		}
 	}
 
 	for (auto Light : LightScene.PointLights.Elements)
 	{
-		if (LightParameters.Count < MaxLightCount)
+		if (LightCount < MaxLightCount)
 		{
-			LightParameters.Position[LightParameters.Count] = Light.Position;
-			LightParameters.Color[LightParameters.Count] = FVector(Light.Color);
-			LightParameters.Dimensions[LightParameters.Count] = FVector(0.0, 0.0, Light.SourceRadius);
-			LightParameters.Attenuation[LightParameters.Count] = 1.0f / Light.AttenuationRadius;
-			LightParameters.FalloffExponent[LightParameters.Count] = Light.FalloffExponent;
-			LightParameters.IESTextureSlice[LightParameters.Count] = -1;
-
-			LightParameters.Flags[LightParameters.Count] = PATHTRACER_FLAG_TRANSMISSION_MASK;
-			LightParameters.Flags[LightParameters.Count] |= PATHTRACER_FLAG_LIGHTING_CHANNEL_MASK;
-			LightParameters.Flags[LightParameters.Count] |= Light.IsInverseSquared ? 0 : PATHTRACER_FLAG_NON_INVERSE_SQUARE_FALLOFF_MASK;
-			LightParameters.Flags[LightParameters.Count] |= Light.bStationary ? PATHTRACER_FLAG_STATIONARY_MASK : 0;
-			LightParameters.Flags[LightParameters.Count] |= PATHTRACING_LIGHT_POINT;
-
-			LightParameters.Count++;
+			FPathTracingLight& DestLight = Lights[LightCount++];
+
+			DestLight.Position = Light.Position;
+			DestLight.Color = FVector(Light.Color);
+			DestLight.Dimensions = FVector(0.0, 0.0, Light.SourceRadius);
+			DestLight.Attenuation = 1.0f / Light.AttenuationRadius;
+			DestLight.FalloffExponent = Light.FalloffExponent;
+			DestLight.IESTextureSlice = -1;
+
+			DestLight.Flags = PATHTRACER_FLAG_TRANSMISSION_MASK;
+			DestLight.Flags |= PATHTRACER_FLAG_LIGHTING_CHANNEL_MASK;
+			DestLight.Flags |= Light.IsInverseSquared ? 0 : PATHTRACER_FLAG_NON_INVERSE_SQUARE_FALLOFF_MASK;
+			DestLight.Flags |= Light.bStationary ? PATHTRACER_FLAG_STATIONARY_MASK : 0;
+			DestLight.Flags |= PATHTRACING_LIGHT_POINT;
 		}
 	}
 
 	for (auto Light : LightScene.SpotLights.Elements)
 	{
-		if (LightParameters.Count < MaxLightCount)
+		if (LightCount < MaxLightCount)
 		{
-			LightParameters.Position[LightParameters.Count] = Light.Position;
-			LightParameters.Normal[LightParameters.Count] = Light.Direction;
-			LightParameters.Color[LightParameters.Count] = FVector(Light.Color);
-			LightParameters.Dimensions[LightParameters.Count] = FVector(Light.SpotAngles, Light.SourceRadius);
-			LightParameters.Attenuation[LightParameters.Count] = 1.0f / Light.AttenuationRadius;
-			LightParameters.FalloffExponent[LightParameters.Count] = Light.FalloffExponent;
-			LightParameters.IESTextureSlice[LightParameters.Count] = -1;
-
-			LightParameters.Flags[LightParameters.Count] = PATHTRACER_FLAG_TRANSMISSION_MASK;
-			LightParameters.Flags[LightParameters.Count] |= PATHTRACER_FLAG_LIGHTING_CHANNEL_MASK;
-			LightParameters.Flags[LightParameters.Count] |= Light.IsInverseSquared ? 0 : PATHTRACER_FLAG_NON_INVERSE_SQUARE_FALLOFF_MASK;
-			LightParameters.Flags[LightParameters.Count] |= Light.bStationary ? PATHTRACER_FLAG_STATIONARY_MASK : 0;
-			LightParameters.Flags[LightParameters.Count] |= PATHTRACING_LIGHT_SPOT;
-
-			LightParameters.Count++;
+			FPathTracingLight& DestLight = Lights[LightCount++];
+
+			DestLight.Position = Light.Position;
+			DestLight.Normal = Light.Direction;
+			DestLight.Color = FVector(Light.Color);
+			DestLight.Dimensions = FVector(Light.SpotAngles, Light.SourceRadius);
+			DestLight.Attenuation = 1.0f / Light.AttenuationRadius;
+			DestLight.FalloffExponent = Light.FalloffExponent;
+			DestLight.IESTextureSlice = -1;
+
+			DestLight.Flags = PATHTRACER_FLAG_TRANSMISSION_MASK;
+			DestLight.Flags |= PATHTRACER_FLAG_LIGHTING_CHANNEL_MASK;
+			DestLight.Flags |= Light.IsInverseSquared ? 0 : PATHTRACER_FLAG_NON_INVERSE_SQUARE_FALLOFF_MASK;
+			DestLight.Flags |= Light.bStationary ? PATHTRACER_FLAG_STATIONARY_MASK : 0;
+			DestLight.Flags |= PATHTRACING_LIGHT_SPOT;
 		}
 	}
 
 	for (auto Light : LightScene.RectLights.Elements)
 	{
-		if (LightParameters.Count < MaxLightCount)
+		if (LightCount < MaxLightCount)
 		{
-			LightParameters.Position[LightParameters.Count] = Light.Position;
-			LightParameters.Normal[LightParameters.Count] = Light.Direction;
-			LightParameters.dPdu[LightParameters.Count] = FVector::CrossProduct(Light.Tangent, -Light.Direction);
-			LightParameters.dPdv[LightParameters.Count] = Light.Tangent;
+			FPathTracingLight& DestLight = Lights[LightCount++];
+
+			DestLight.Position = Light.Position;
+			DestLight.Normal = Light.Direction;
+			DestLight.dPdu = FVector::CrossProduct(Light.Tangent, -Light.Direction);
+			DestLight.dPdv = Light.Tangent;
 
 			FLinearColor LightColor = Light.Color;
 			LightColor /= 0.5f * Light.SourceWidth * Light.SourceHeight;
-			LightParameters.Color[LightParameters.Count] = FVector(LightColor);
-
-			LightParameters.Dimensions[LightParameters.Count] = FVector(Light.SourceWidth, Light.SourceHeight, 0.0f);
-			LightParameters.Attenuation[LightParameters.Count] = 1.0f / Light.AttenuationRadius;
-			LightParameters.RectLightBarnCosAngle[LightParameters.Count] = FMath::Cos(FMath::DegreesToRadians(Light.BarnDoorAngle));
-			LightParameters.RectLightBarnLength[LightParameters.Count] = Light.BarnDoorLength;
+			DestLight.Color = FVector(LightColor);
 
-			LightParameters.IESTextureSlice[LightParameters.Count] = -1;
+			DestLight.Dimensions = FVector(Light.SourceWidth, Light.SourceHeight, 0.0f);
+			DestLight.Attenuation = 1.0f / Light.AttenuationRadius;
+			DestLight.RectLightBarnCosAngle = FMath::Cos(FMath::DegreesToRadians(Light.BarnDoorAngle));
+			DestLight.RectLightBarnLength = Light.BarnDoorLength;
 
-			LightParameters.Flags[LightParameters.Count] = PATHTRACER_FLAG_TRANSMISSION_MASK;
-			LightParameters.Flags[LightParameters.Count] |= PATHTRACER_FLAG_LIGHTING_CHANNEL_MASK;
-			LightParameters.Flags[LightParameters.Count] |= Light.bStationary ? PATHTRACER_FLAG_STATIONARY_MASK : 0;
-			LightParameters.Flags[LightParameters.Count] |= PATHTRACING_LIGHT_RECT;
+			DestLight.IESTextureSlice = -1;
 
-			LightParameters.Count++;
+			DestLight.Flags = PATHTRACER_FLAG_TRANSMISSION_MASK;
+			DestLight.Flags |= PATHTRACER_FLAG_LIGHTING_CHANNEL_MASK;
+			DestLight.Flags |= Light.bStationary ? PATHTRACER_FLAG_STATIONARY_MASK : 0;
+			DestLight.Flags |= PATHTRACING_LIGHT_RECT;
 		}
 	}
 
-	return LightParameters;
+	{
+		// Upload the buffer of lights to the GPU
+		size_t DataSize = sizeof(FPathTracingLight) * FMath::Max(LightCount, 1u);
+		*OutLightBuffer = GraphBuilder.CreateSRV(FRDGBufferSRVDesc(CreateStructuredBuffer(GraphBuilder, TEXT("PathTracingLightsBuffer"), sizeof(FPathTracingLight), FMath::Max(LightCount, 1u), Lights, DataSize)));
+		*OutLightCount = LightCount;
+	}
 }
 
 FSkyLightData SetupSkyLightParameters(const GPULightmass::FLightSceneRenderState& LightScene)
@@ -1699,7 +1701,6 @@ void FLightmapRenderer::Finalize(FRDGBuilder& GraphBuilder)
 								}
 
 								// These two buffers must have lifetime extended beyond GraphBuilder.Execute()
-								TUniformBufferRef<FPathTracingLightData> LightDataUniformBuffer;
 								TUniformBufferRef<FSkyLightData> SkyLightDataUniformBuffer;
 
 								FIntPoint RayTracingResolution;
@@ -1733,8 +1734,8 @@ void FLightmapRenderer::Finalize(FRDGBuilder& GraphBuilder)
 										PassParameters->IrradianceCachingParameters = Scene->IrradianceCache->IrradianceCachingParametersUniformBuffer;
 
 										{
-											LightDataUniformBuffer = CreateUniformBufferImmediate(SetupPathTracingLightParameters(Scene->LightSceneRenderState), EUniformBufferUsage::UniformBuffer_SingleFrame);
-											PassParameters->LightParameters = LightDataUniformBuffer;
+											
+											SetupPathTracingLightParameters(Scene->LightSceneRenderState, GraphBuilder, &PassParameters->SceneLights, &PassParameters->SceneLightCount);
 										}
 
 										{
@@ -2783,7 +2784,7 @@ void FLightmapRenderer::BackgroundTick()
 
 		if (Mip0WorkDoneLastFrame < NumWorkPerFrame)
 		{
-			int32 PoolSize = FMath::CeilToInt(FMath::Sqrt(NumWorkPerFrame * 3.f));
+			int32 PoolSize = FMath::CeilToInt(FMath::Sqrt(static_cast<float>(NumWorkPerFrame * 3.f)));
 
 			FIntPoint TextureSize(PoolSize * GPreviewLightmapPhysicalTileSize, PoolSize * GPreviewLightmapPhysicalTileSize);
 
diff --git a/Engine/Plugins/Experimental/GPULightmass/Source/GPULightmass/Private/VolumetricLightmap.cpp b/Engine/Plugins/Experimental/GPULightmass/Source/GPULightmass/Private/VolumetricLightmap.cpp
--- a/Engine/Plugins/Experimental/GPULightmass/Source/GPULightmass/Private/VolumetricLightmap.cpp
+++ b/Engine/Plugins/Experimental/GPULightmass/Source/GPULightmass/Private/VolumetricLightmap.cpp
@@ -74,7 +74,7 @@ void ReleaseBrickData(FVolumetricLightmapBrickData& BrickData)
 	BrickData.DirectionalLightShadowing.UAV.SafeRelease();
 }
 
-FPathTracingLightData SetupPathTracingLightParameters(const GPULightmass::FLightSceneRenderState& LightScene);
+void SetupPathTracingLightParameters(const GPULightmass::FLightSceneRenderState& LightScene, FRDGBuilder& GraphBuilder, FRDGBufferSRV** OutLightBuffer, uint32* OutLightCount);
 FSkyLightData SetupSkyLightParameters(const GPULightmass::FLightSceneRenderState& LightScene);
 
 namespace GPULightmass
@@ -482,21 +482,28 @@ void FVolumetricLightmapRenderer::BackgroundTick()
 		Scene->SetupRayTracingScene();
 	}
 
-	FRHICommandListImmediate& RHICmdList = FRHICommandListExecutor::GetImmediateCommandList();
+	FRDGBuilder GraphBuilder(FRHICommandListExecutor::GetImmediateCommandList());
+	RDG_EVENT_SCOPE(GraphBuilder, "Volumetric Lightmap Path Tracing");
 
 	bool bLastFewFramesIdle = !GCurrentLevelEditingViewportClient || !GCurrentLevelEditingViewportClient->IsRealtime();
 
 	int32 NumSamplesThisFrame = !bLastFewFramesIdle ? 1 : 32;
 
-	RHICmdList.Transition({
-		FRHITransitionInfo(VolumetricLightmapData.BrickData.AmbientVector.UAV, ERHIAccess::Unknown, ERHIAccess::UAVCompute),
-		FRHITransitionInfo(VolumetricLightmapData.BrickData.SHCoefficients[0].UAV, ERHIAccess::Unknown, ERHIAccess::UAVCompute),
-		FRHITransitionInfo(VolumetricLightmapData.BrickData.SHCoefficients[1].UAV, ERHIAccess::Unknown, ERHIAccess::UAVCompute),
-		FRHITransitionInfo(VolumetricLightmapData.BrickData.SHCoefficients[2].UAV, ERHIAccess::Unknown, ERHIAccess::UAVCompute),
-		FRHITransitionInfo(VolumetricLightmapData.BrickData.SHCoefficients[3].UAV, ERHIAccess::Unknown, ERHIAccess::UAVCompute),
-		FRHITransitionInfo(VolumetricLightmapData.BrickData.SHCoefficients[4].UAV, ERHIAccess::Unknown, ERHIAccess::UAVCompute),
-		FRHITransitionInfo(VolumetricLightmapData.BrickData.SHCoefficients[5].UAV, ERHIAccess::Unknown, ERHIAccess::UAVCompute)
-	});
+	// manually handle transitions since the buffers are not (yet) managed by RDG
+	GraphBuilder.AddPass(RDG_EVENT_NAME("Transition Buffers"), ERDGPassFlags::None, 
+		[this](FRHICommandListImmediate& RHICmdList) {
+			RHICmdList.Transition({
+				FRHITransitionInfo(VolumetricLightmapData.BrickData.AmbientVector.UAV, ERHIAccess::Unknown, ERHIAccess::UAVCompute),
+				FRHITransitionInfo(VolumetricLightmapData.BrickData.SHCoefficients[0].UAV, ERHIAccess::Unknown, ERHIAccess::UAVCompute),
+				FRHITransitionInfo(VolumetricLightmapData.BrickData.SHCoefficients[1].UAV, ERHIAccess::Unknown, ERHIAccess::UAVCompute),
+				FRHITransitionInfo(VolumetricLightmapData.BrickData.SHCoefficients[2].UAV, ERHIAccess::Unknown, ERHIAccess::UAVCompute),
+				FRHITransitionInfo(VolumetricLightmapData.BrickData.SHCoefficients[3].UAV, ERHIAccess::Unknown, ERHIAccess::UAVCompute),
+				FRHITransitionInfo(VolumetricLightmapData.BrickData.SHCoefficients[4].UAV, ERHIAccess::Unknown, ERHIAccess::UAVCompute),
+				FRHITransitionInfo(VolumetricLightmapData.BrickData.SHCoefficients[5].UAV, ERHIAccess::Unknown, ERHIAccess::UAVCompute)
+				}
+			);
+		}
+	);
 
 	for (int32 SampleIndex = 0; SampleIndex < NumSamplesThisFrame; SampleIndex++)
 	{
@@ -509,10 +516,7 @@ void FVolumetricLightmapRenderer::BackgroundTick()
 #if RHI_RAYTRACING
 		if (IsRayTracingEnabled())
 		{
-			SCOPED_DRAW_EVENTF(RHICmdList, VolumetricLightmapPathTracing, TEXT("VolumetricLightmapPathTracing %d bricks %d rays"), BricksToCalcThisFrame, BricksToCalcThisFrame * (BrickSize + 1) * (BrickSize + 1) * (BrickSize + 1));
-
 			// These two buffers must have lifetime extended beyond RHICmdList.RayTraceDispatch()
-			TUniformBufferRef<FPathTracingLightData> LightDataUniformBuffer;
 			TUniformBufferRef<FSkyLightData> SkyLightDataUniformBuffer;
 
 			FGlobalShaderMap* GlobalShaderMap = GetGlobalShaderMap(GMaxRHIFeatureLevel);
@@ -522,104 +526,119 @@ void FVolumetricLightmapRenderer::BackgroundTick()
 				PermutationVector.Set<FVolumetricLightmapPathTracingRGS::FUseIrradianceCaching>(Scene->Settings->bUseIrradianceCaching);
 				TShaderRef<FVolumetricLightmapPathTracingRGS> RayGenShader = GlobalShaderMap->GetShader<FVolumetricLightmapPathTracingRGS>(PermutationVector);
 
-				FVolumetricLightmapPathTracingRGS::FParameters Parameters;
-				Parameters.FrameNumber = FrameNumber / NumFramesOneRound;
-				Parameters.VolumeMin = VolumeMin;
-				Parameters.VolumeSize = VolumeSize;
-				Parameters.IndirectionTextureDim = IndirectionTextureDimensions;
-				Parameters.TLAS = Scene->RayTracingScene->GetShaderResourceView();
-				Parameters.BrickRequests = BrickRequests.SRV;
-				Parameters.NumTotalBricks = NumTotalBricks;
-				Parameters.BrickBatchOffset = BrickBatchOffset;
-				Parameters.AmbientVector = AccumulationBrickData.AmbientVector.UAV;
-				Parameters.SHCoefficients0R = AccumulationBrickData.SHCoefficients[0].UAV;
-				Parameters.SHCoefficients1R = AccumulationBrickData.SHCoefficients[1].UAV;
-				Parameters.SHCoefficients0G = AccumulationBrickData.SHCoefficients[2].UAV;
-				Parameters.SHCoefficients1G = AccumulationBrickData.SHCoefficients[3].UAV;
-				Parameters.SHCoefficients0B = AccumulationBrickData.SHCoefficients[4].UAV;
-				Parameters.SHCoefficients1B = AccumulationBrickData.SHCoefficients[5].UAV;
-				Parameters.OutAmbientVector = VolumetricLightmapData.BrickData.AmbientVector.UAV;
-				Parameters.OutSHCoefficients0R = VolumetricLightmapData.BrickData.SHCoefficients[0].UAV;
-				Parameters.OutSHCoefficients1R = VolumetricLightmapData.BrickData.SHCoefficients[1].UAV;
-				Parameters.OutSHCoefficients0G = VolumetricLightmapData.BrickData.SHCoefficients[2].UAV;
-				Parameters.OutSHCoefficients1G = VolumetricLightmapData.BrickData.SHCoefficients[3].UAV;
-				Parameters.OutSHCoefficients0B = VolumetricLightmapData.BrickData.SHCoefficients[4].UAV;
-				Parameters.OutSHCoefficients1B = VolumetricLightmapData.BrickData.SHCoefficients[5].UAV;
-				Parameters.ViewUniformBuffer = Scene->ReferenceView->ViewUniformBuffer;
-				Parameters.IrradianceCachingParameters = Scene->IrradianceCache->IrradianceCachingParametersUniformBuffer;
+				FVolumetricLightmapPathTracingRGS::FParameters* PassParameters = GraphBuilder.AllocParameters<FVolumetricLightmapPathTracingRGS::FParameters>();
+				PassParameters->FrameNumber = FrameNumber / NumFramesOneRound;
+				PassParameters->VolumeMin = VolumeMin;
+				PassParameters->VolumeSize = VolumeSize;
+				PassParameters->IndirectionTextureDim = IndirectionTextureDimensions;
+				PassParameters->TLAS = Scene->RayTracingScene->GetShaderResourceView();
+				PassParameters->BrickRequests = BrickRequests.SRV;
+				PassParameters->NumTotalBricks = NumTotalBricks;
+				PassParameters->BrickBatchOffset = BrickBatchOffset;
+				PassParameters->AmbientVector = AccumulationBrickData.AmbientVector.UAV;
+				PassParameters->SHCoefficients0R = AccumulationBrickData.SHCoefficients[0].UAV;
+				PassParameters->SHCoefficients1R = AccumulationBrickData.SHCoefficients[1].UAV;
+				PassParameters->SHCoefficients0G = AccumulationBrickData.SHCoefficients[2].UAV;
+				PassParameters->SHCoefficients1G = AccumulationBrickData.SHCoefficients[3].UAV;
+				PassParameters->SHCoefficients0B = AccumulationBrickData.SHCoefficients[4].UAV;
+				PassParameters->SHCoefficients1B = AccumulationBrickData.SHCoefficients[5].UAV;
+				PassParameters->OutAmbientVector = VolumetricLightmapData.BrickData.AmbientVector.UAV;
+				PassParameters->OutSHCoefficients0R = VolumetricLightmapData.BrickData.SHCoefficients[0].UAV;
+				PassParameters->OutSHCoefficients1R = VolumetricLightmapData.BrickData.SHCoefficients[1].UAV;
+				PassParameters->OutSHCoefficients0G = VolumetricLightmapData.BrickData.SHCoefficients[2].UAV;
+				PassParameters->OutSHCoefficients1G = VolumetricLightmapData.BrickData.SHCoefficients[3].UAV;
+				PassParameters->OutSHCoefficients0B = VolumetricLightmapData.BrickData.SHCoefficients[4].UAV;
+				PassParameters->OutSHCoefficients1B = VolumetricLightmapData.BrickData.SHCoefficients[5].UAV;
+				PassParameters->ViewUniformBuffer = Scene->ReferenceView->ViewUniformBuffer;
+				PassParameters->IrradianceCachingParameters = Scene->IrradianceCache->IrradianceCachingParametersUniformBuffer;
+
+				SetupPathTracingLightParameters(Scene->LightSceneRenderState, GraphBuilder, &PassParameters->SceneLights, &PassParameters->SceneLightCount);
+
+				PassParameters->SkyLight = CreateUniformBufferImmediate(SetupSkyLightParameters(Scene->LightSceneRenderState), EUniformBufferUsage::UniformBuffer_SingleFrame);
 
-				{
-					LightDataUniformBuffer = CreateUniformBufferImmediate(SetupPathTracingLightParameters(Scene->LightSceneRenderState), EUniformBufferUsage::UniformBuffer_SingleFrame);
-					Parameters.LightParameters = LightDataUniformBuffer;
-				}
+				// TODO: find a way to share IES atlas with path tracer ...
+				PassParameters->IESTexture = GWhiteTexture->TextureRHI;
+				PassParameters->IESTextureSampler = TStaticSamplerState<SF_Bilinear, AM_Clamp, AM_Clamp, AM_Clamp>::GetRHI();
 
-				{
-					SkyLightDataUniformBuffer = CreateUniformBufferImmediate(SetupSkyLightParameters(Scene->LightSceneRenderState), EUniformBufferUsage::UniformBuffer_SingleFrame);
-					Parameters.SkyLight = SkyLightDataUniformBuffer;
-				}
+				FSceneRenderState* SceneRenderState = Scene; // capture member variable
+				GraphBuilder.AddPass(
+					RDG_EVENT_NAME("VolumetricLightmapPathTracing %d bricks %d rays", BricksToCalcThisFrame, BricksToCalcThisFrame * (BrickSize + 1) * (BrickSize + 1) * (BrickSize + 1)),
+					PassParameters,
+					ERDGPassFlags::Compute,
+					[PassParameters, RayGenShader, SceneRenderState, BricksToCalcThisFrame](FRHICommandListImmediate& RHICmdList)
+					{
+						FRayTracingShaderBindingsWriter GlobalResources;
+						SetShaderParameters(GlobalResources, RayGenShader, *PassParameters);
 
-				// TODO: find a way to share IES atlas with path tracer ...
-				Parameters.IESTexture = GWhiteTexture->TextureRHI;
-				Parameters.IESTextureSampler = TStaticSamplerState<SF_Bilinear, AM_Clamp, AM_Clamp, AM_Clamp>::GetRHI();
+						FRHIRayTracingScene* RayTracingSceneRHI = SceneRenderState->RayTracingScene;
+						RHICmdList.RayTraceDispatch(SceneRenderState->RayTracingPipelineState, RayGenShader.GetRayTracingShader(), RayTracingSceneRHI, GlobalResources, BricksToCalcThisFrame * (BrickSize + 1)* (BrickSize + 1)* (BrickSize + 1), 1);
+					}
+				);
 
-				FRayTracingShaderBindingsWriter GlobalResources;
-				SetShaderParameters(GlobalResources, RayGenShader, Parameters);
 
-				FRHIRayTracingScene* RayTracingSceneRHI = Scene->RayTracingScene;
-				RHICmdList.RayTraceDispatch(Scene->RayTracingPipelineState, RayGenShader.GetRayTracingShader(), RayTracingSceneRHI, GlobalResources, BricksToCalcThisFrame * (BrickSize + 1) * (BrickSize + 1) * (BrickSize + 1), 1);
 			}
 		}
 #endif
 		{
-			RHICmdList.Transition({
-				FRHITransitionInfo(AccumulationBrickData.AmbientVector.UAV, ERHIAccess::UAVCompute, ERHIAccess::UAVCompute),
-				FRHITransitionInfo(AccumulationBrickData.SHCoefficients[0].UAV, ERHIAccess::UAVCompute, ERHIAccess::UAVCompute),
-				FRHITransitionInfo(AccumulationBrickData.SHCoefficients[1].UAV, ERHIAccess::UAVCompute, ERHIAccess::UAVCompute),
-				FRHITransitionInfo(AccumulationBrickData.SHCoefficients[2].UAV, ERHIAccess::UAVCompute, ERHIAccess::UAVCompute),
-				FRHITransitionInfo(AccumulationBrickData.SHCoefficients[3].UAV, ERHIAccess::UAVCompute, ERHIAccess::UAVCompute),
-				FRHITransitionInfo(AccumulationBrickData.SHCoefficients[4].UAV, ERHIAccess::UAVCompute, ERHIAccess::UAVCompute),
-				FRHITransitionInfo(AccumulationBrickData.SHCoefficients[5].UAV, ERHIAccess::UAVCompute, ERHIAccess::UAVCompute),
-			});
+			// manually handle transitions since the buffers are not (yet) managed by RDG
+			GraphBuilder.AddPass(RDG_EVENT_NAME("Transition Buffers"), ERDGPassFlags::None,
+				[this](FRHICommandListImmediate& RHICmdList) {
+					RHICmdList.Transition({
+						FRHITransitionInfo(AccumulationBrickData.AmbientVector.UAV, ERHIAccess::UAVCompute, ERHIAccess::UAVCompute),
+						FRHITransitionInfo(AccumulationBrickData.SHCoefficients[0].UAV, ERHIAccess::UAVCompute, ERHIAccess::UAVCompute),
+						FRHITransitionInfo(AccumulationBrickData.SHCoefficients[1].UAV, ERHIAccess::UAVCompute, ERHIAccess::UAVCompute),
+						FRHITransitionInfo(AccumulationBrickData.SHCoefficients[2].UAV, ERHIAccess::UAVCompute, ERHIAccess::UAVCompute),
+						FRHITransitionInfo(AccumulationBrickData.SHCoefficients[3].UAV, ERHIAccess::UAVCompute, ERHIAccess::UAVCompute),
+						FRHITransitionInfo(AccumulationBrickData.SHCoefficients[4].UAV, ERHIAccess::UAVCompute, ERHIAccess::UAVCompute),
+						FRHITransitionInfo(AccumulationBrickData.SHCoefficients[5].UAV, ERHIAccess::UAVCompute, ERHIAccess::UAVCompute),
+						}
+					);
+				}
+			);
 
 			FGlobalShaderMap* GlobalShaderMap = GetGlobalShaderMap(GMaxRHIFeatureLevel);
 
 			TShaderMapRef<FFinalizeBrickResultsCS> ComputeShader(GlobalShaderMap);
 
-			FFinalizeBrickResultsCS::FParameters Parameters;
-			Parameters.NumTotalBricks = NumTotalBricks;
-			Parameters.BrickBatchOffset = BrickBatchOffset;
-			Parameters.BrickRequests = BrickRequests.UAV;
-			Parameters.AmbientVector = AccumulationBrickData.AmbientVector.Texture;
-			Parameters.SHCoefficients0R = AccumulationBrickData.SHCoefficients[0].Texture;
-			Parameters.SHCoefficients1R = AccumulationBrickData.SHCoefficients[1].Texture;
-			Parameters.SHCoefficients0G = AccumulationBrickData.SHCoefficients[2].Texture;
-			Parameters.SHCoefficients1G = AccumulationBrickData.SHCoefficients[3].Texture;
-			Parameters.SHCoefficients0B = AccumulationBrickData.SHCoefficients[4].Texture;
-			Parameters.SHCoefficients1B = AccumulationBrickData.SHCoefficients[5].Texture;
-			Parameters.OutAmbientVector = VolumetricLightmapData.BrickData.AmbientVector.UAV;
-			Parameters.OutSHCoefficients0R = VolumetricLightmapData.BrickData.SHCoefficients[0].UAV;
-			Parameters.OutSHCoefficients1R = VolumetricLightmapData.BrickData.SHCoefficients[1].UAV;
-			Parameters.OutSHCoefficients0G = VolumetricLightmapData.BrickData.SHCoefficients[2].UAV;
-			Parameters.OutSHCoefficients1G = VolumetricLightmapData.BrickData.SHCoefficients[3].UAV;
-			Parameters.OutSHCoefficients0B = VolumetricLightmapData.BrickData.SHCoefficients[4].UAV;
-			Parameters.OutSHCoefficients1B = VolumetricLightmapData.BrickData.SHCoefficients[5].UAV;
-
-			FComputeShaderUtils::Dispatch(RHICmdList, ComputeShader, Parameters, FIntVector(BricksToCalcThisFrame, 1, 1));
-
-			RHICmdList.Transition({
-				FRHITransitionInfo(AccumulationBrickData.AmbientVector.UAV, ERHIAccess::UAVCompute, ERHIAccess::UAVCompute),
-				FRHITransitionInfo(AccumulationBrickData.SHCoefficients[0].UAV, ERHIAccess::UAVCompute, ERHIAccess::UAVCompute),
-				FRHITransitionInfo(AccumulationBrickData.SHCoefficients[1].UAV, ERHIAccess::UAVCompute, ERHIAccess::UAVCompute),
-				FRHITransitionInfo(AccumulationBrickData.SHCoefficients[2].UAV, ERHIAccess::UAVCompute, ERHIAccess::UAVCompute),
-				FRHITransitionInfo(AccumulationBrickData.SHCoefficients[3].UAV, ERHIAccess::UAVCompute, ERHIAccess::UAVCompute),
-				FRHITransitionInfo(AccumulationBrickData.SHCoefficients[4].UAV, ERHIAccess::UAVCompute, ERHIAccess::UAVCompute),
-				FRHITransitionInfo(AccumulationBrickData.SHCoefficients[5].UAV, ERHIAccess::UAVCompute, ERHIAccess::UAVCompute),
-			});
+			FFinalizeBrickResultsCS::FParameters* PassParameters = GraphBuilder.AllocParameters<FFinalizeBrickResultsCS::FParameters>();
+			PassParameters->NumTotalBricks = NumTotalBricks;
+			PassParameters->BrickBatchOffset = BrickBatchOffset;
+			PassParameters->BrickRequests = BrickRequests.UAV;
+			PassParameters->AmbientVector = AccumulationBrickData.AmbientVector.Texture;
+			PassParameters->SHCoefficients0R = AccumulationBrickData.SHCoefficients[0].Texture;
+			PassParameters->SHCoefficients1R = AccumulationBrickData.SHCoefficients[1].Texture;
+			PassParameters->SHCoefficients0G = AccumulationBrickData.SHCoefficients[2].Texture;
+			PassParameters->SHCoefficients1G = AccumulationBrickData.SHCoefficients[3].Texture;
+			PassParameters->SHCoefficients0B = AccumulationBrickData.SHCoefficients[4].Texture;
+			PassParameters->SHCoefficients1B = AccumulationBrickData.SHCoefficients[5].Texture;
+			PassParameters->OutAmbientVector = VolumetricLightmapData.BrickData.AmbientVector.UAV;
+			PassParameters->OutSHCoefficients0R = VolumetricLightmapData.BrickData.SHCoefficients[0].UAV;
+			PassParameters->OutSHCoefficients1R = VolumetricLightmapData.BrickData.SHCoefficients[1].UAV;
+			PassParameters->OutSHCoefficients0G = VolumetricLightmapData.BrickData.SHCoefficients[2].UAV;
+			PassParameters->OutSHCoefficients1G = VolumetricLightmapData.BrickData.SHCoefficients[3].UAV;
+			PassParameters->OutSHCoefficients0B = VolumetricLightmapData.BrickData.SHCoefficients[4].UAV;
+			PassParameters->OutSHCoefficients1B = VolumetricLightmapData.BrickData.SHCoefficients[5].UAV;
+
+			FComputeShaderUtils::AddPass(GraphBuilder, RDG_EVENT_NAME("FinalizeBrickResults"), ComputeShader, PassParameters, FIntVector(BricksToCalcThisFrame, 1, 1));
+
+			// manually handle transitions since the buffers are not (yet) managed by RDG
+			GraphBuilder.AddPass(RDG_EVENT_NAME("Transition Buffers"), ERDGPassFlags::None,
+				[this](FRHICommandListImmediate& RHICmdList) {
+					RHICmdList.Transition({
+						FRHITransitionInfo(AccumulationBrickData.AmbientVector.UAV, ERHIAccess::UAVCompute, ERHIAccess::UAVCompute),
+						FRHITransitionInfo(AccumulationBrickData.SHCoefficients[0].UAV, ERHIAccess::UAVCompute, ERHIAccess::UAVCompute),
+						FRHITransitionInfo(AccumulationBrickData.SHCoefficients[1].UAV, ERHIAccess::UAVCompute, ERHIAccess::UAVCompute),
+						FRHITransitionInfo(AccumulationBrickData.SHCoefficients[2].UAV, ERHIAccess::UAVCompute, ERHIAccess::UAVCompute),
+						FRHITransitionInfo(AccumulationBrickData.SHCoefficients[3].UAV, ERHIAccess::UAVCompute, ERHIAccess::UAVCompute),
+						FRHITransitionInfo(AccumulationBrickData.SHCoefficients[4].UAV, ERHIAccess::UAVCompute, ERHIAccess::UAVCompute),
+						FRHITransitionInfo(AccumulationBrickData.SHCoefficients[5].UAV, ERHIAccess::UAVCompute, ERHIAccess::UAVCompute),
+						}
+					);
+				}
+			);
 		}
 
 		{
-			SCOPED_DRAW_EVENTF(RHICmdList, VolumetricLightmapStitching, TEXT("VolumetricLightmapStitching %d bricks"), BricksToCalcThisFrame);
-
 			// Doing 2 passes no longer makes sense in an amortized setup
 			// for (int32 StitchPass = 0; StitchPass < 2; StitchPass++)
 			{
@@ -627,34 +646,40 @@ void FVolumetricLightmapRenderer::BackgroundTick()
 
 				TShaderMapRef<FStitchBorderCS> ComputeShader(GlobalShaderMap);
 
-				FStitchBorderCS::FParameters Parameters;
-				Parameters.BrickDataDimensions = VolumetricLightmapData.BrickDataDimensions;
-				Parameters.IndirectionTextureDim = IndirectionTextureDimensions;
-				Parameters.FrameNumber = FrameNumber / NumFramesOneRound;
-				Parameters.NumTotalBricks = NumTotalBricks;
-				Parameters.BrickBatchOffset = BrickBatchOffset;
-				Parameters.IndirectionTexture = IndirectionTexture->GetRenderTargetItem().UAV;
-				Parameters.BrickRequests = BrickRequests.UAV;
-				Parameters.AmbientVector = AccumulationBrickData.AmbientVector.Texture;
-				Parameters.OutAmbientVector = VolumetricLightmapData.BrickData.AmbientVector.UAV;
-				Parameters.OutSHCoefficients0R = VolumetricLightmapData.BrickData.SHCoefficients[0].UAV;
-				Parameters.OutSHCoefficients1R = VolumetricLightmapData.BrickData.SHCoefficients[1].UAV;
-				Parameters.OutSHCoefficients0G = VolumetricLightmapData.BrickData.SHCoefficients[2].UAV;
-				Parameters.OutSHCoefficients1G = VolumetricLightmapData.BrickData.SHCoefficients[3].UAV;
-				Parameters.OutSHCoefficients0B = VolumetricLightmapData.BrickData.SHCoefficients[4].UAV;
-				Parameters.OutSHCoefficients1B = VolumetricLightmapData.BrickData.SHCoefficients[5].UAV;
-
-				FComputeShaderUtils::Dispatch(RHICmdList, ComputeShader, Parameters, FIntVector(BricksToCalcThisFrame, 1, 1));
-
-				RHICmdList.Transition({
-					FRHITransitionInfo(AccumulationBrickData.AmbientVector.UAV, ERHIAccess::UAVCompute, ERHIAccess::UAVCompute),
-					FRHITransitionInfo(AccumulationBrickData.SHCoefficients[0].UAV, ERHIAccess::UAVCompute, ERHIAccess::UAVCompute),
-					FRHITransitionInfo(AccumulationBrickData.SHCoefficients[1].UAV, ERHIAccess::UAVCompute, ERHIAccess::UAVCompute),
-					FRHITransitionInfo(AccumulationBrickData.SHCoefficients[2].UAV, ERHIAccess::UAVCompute, ERHIAccess::UAVCompute),
-					FRHITransitionInfo(AccumulationBrickData.SHCoefficients[3].UAV, ERHIAccess::UAVCompute, ERHIAccess::UAVCompute),
-					FRHITransitionInfo(AccumulationBrickData.SHCoefficients[4].UAV, ERHIAccess::UAVCompute, ERHIAccess::UAVCompute),
-					FRHITransitionInfo(AccumulationBrickData.SHCoefficients[5].UAV, ERHIAccess::UAVCompute, ERHIAccess::UAVCompute)
-				});
+				FStitchBorderCS::FParameters* PassParameters = GraphBuilder.AllocParameters<FStitchBorderCS::FParameters>();
+				PassParameters->BrickDataDimensions = VolumetricLightmapData.BrickDataDimensions;
+				PassParameters->IndirectionTextureDim = IndirectionTextureDimensions;
+				PassParameters->FrameNumber = FrameNumber / NumFramesOneRound;
+				PassParameters->NumTotalBricks = NumTotalBricks;
+				PassParameters->BrickBatchOffset = BrickBatchOffset;
+				PassParameters->IndirectionTexture = IndirectionTexture->GetRenderTargetItem().UAV;
+				PassParameters->BrickRequests = BrickRequests.UAV;
+				PassParameters->AmbientVector = AccumulationBrickData.AmbientVector.Texture;
+				PassParameters->OutAmbientVector = VolumetricLightmapData.BrickData.AmbientVector.UAV;
+				PassParameters->OutSHCoefficients0R = VolumetricLightmapData.BrickData.SHCoefficients[0].UAV;
+				PassParameters->OutSHCoefficients1R = VolumetricLightmapData.BrickData.SHCoefficients[1].UAV;
+				PassParameters->OutSHCoefficients0G = VolumetricLightmapData.BrickData.SHCoefficients[2].UAV;
+				PassParameters->OutSHCoefficients1G = VolumetricLightmapData.BrickData.SHCoefficients[3].UAV;
+				PassParameters->OutSHCoefficients0B = VolumetricLightmapData.BrickData.SHCoefficients[4].UAV;
+				PassParameters->OutSHCoefficients1B = VolumetricLightmapData.BrickData.SHCoefficients[5].UAV;
+
+				FComputeShaderUtils::AddPass(GraphBuilder, RDG_EVENT_NAME("VolumetricLightmapStitching %d bricks", BricksToCalcThisFrame), ComputeShader, PassParameters, FIntVector(BricksToCalcThisFrame, 1, 1));
+
+				// manually handle transitions since the buffers are not (yet) managed by RDG
+				GraphBuilder.AddPass(RDG_EVENT_NAME("Transition Buffers"), ERDGPassFlags::None,
+					[this](FRHICommandListImmediate& RHICmdList) {
+						RHICmdList.Transition({
+							FRHITransitionInfo(AccumulationBrickData.AmbientVector.UAV, ERHIAccess::UAVCompute, ERHIAccess::UAVCompute),
+							FRHITransitionInfo(AccumulationBrickData.SHCoefficients[0].UAV, ERHIAccess::UAVCompute, ERHIAccess::UAVCompute),
+							FRHITransitionInfo(AccumulationBrickData.SHCoefficients[1].UAV, ERHIAccess::UAVCompute, ERHIAccess::UAVCompute),
+							FRHITransitionInfo(AccumulationBrickData.SHCoefficients[2].UAV, ERHIAccess::UAVCompute, ERHIAccess::UAVCompute),
+							FRHITransitionInfo(AccumulationBrickData.SHCoefficients[3].UAV, ERHIAccess::UAVCompute, ERHIAccess::UAVCompute),
+							FRHITransitionInfo(AccumulationBrickData.SHCoefficients[4].UAV, ERHIAccess::UAVCompute, ERHIAccess::UAVCompute),
+							FRHITransitionInfo(AccumulationBrickData.SHCoefficients[5].UAV, ERHIAccess::UAVCompute, ERHIAccess::UAVCompute)
+							}
+						);
+					}
+				);
 			}
 		}
 
@@ -667,6 +692,7 @@ void FVolumetricLightmapRenderer::BackgroundTick()
 			break;
 		}
 	}
+	GraphBuilder.Execute();
 
 	if (IsRayTracingEnabled())
 	{
diff --git a/Engine/Shaders/Private/PathTracing/Light/PathTracingLightCommon.ush b/Engine/Shaders/Private/PathTracing/Light/PathTracingLightCommon.ush
--- a/Engine/Shaders/Private/PathTracing/Light/PathTracingLightCommon.ush
+++ b/Engine/Shaders/Private/PathTracing/Light/PathTracingLightCommon.ush
@@ -6,6 +6,12 @@ LightCommon.usf: Common utility functions for light sampling
 
 #pragma once
 
+#include "/Engine/Shared/RayTracingTypes.h"
+
+uint SceneLightCount;
+StructuredBuffer<FPathTracingLight> SceneLights;
+
+
 int GetLightId(FMaterialClosestHitPayload RayHitInfo)
 {
 	return -1;
@@ -13,27 +19,27 @@ int GetLightId(FMaterialClosestHitPayload RayHitInfo)
 
 bool IsEnvironmentLight(int LightId)
 {
-	return (SceneLightsData.Flags[LightId] & PATHTRACER_FLAG_TYPE_MASK) == PATHTRACING_LIGHT_SKY;
+	return (SceneLights[LightId].Flags & PATHTRACER_FLAG_TYPE_MASK) == PATHTRACING_LIGHT_SKY;
 }
 
 bool IsPointLight(int LightId)
 {
-	return (SceneLightsData.Flags[LightId] & PATHTRACER_FLAG_TYPE_MASK) == PATHTRACING_LIGHT_POINT;
+	return (SceneLights[LightId].Flags & PATHTRACER_FLAG_TYPE_MASK) == PATHTRACING_LIGHT_POINT;
 }
 
 bool IsDirectionalLight(int LightId)
 {
-	return (SceneLightsData.Flags[LightId] & PATHTRACER_FLAG_TYPE_MASK) == PATHTRACING_LIGHT_DIRECTIONAL;
+	return (SceneLights[LightId].Flags & PATHTRACER_FLAG_TYPE_MASK) == PATHTRACING_LIGHT_DIRECTIONAL;
 }
 
 bool IsRectLight(int LightId)
 {
-	return (SceneLightsData.Flags[LightId] & PATHTRACER_FLAG_TYPE_MASK) == PATHTRACING_LIGHT_RECT;
+	return (SceneLights[LightId].Flags & PATHTRACER_FLAG_TYPE_MASK) == PATHTRACING_LIGHT_RECT;
 }
 
 bool IsSpotLight(int LightId)
 {
-	return (SceneLightsData.Flags[LightId] & PATHTRACER_FLAG_TYPE_MASK) == PATHTRACING_LIGHT_SPOT;
+	return (SceneLights[LightId].Flags & PATHTRACER_FLAG_TYPE_MASK) == PATHTRACING_LIGHT_SPOT;
 }
 
 // A light is a physical light if it can be intersected by a ray.
@@ -54,72 +60,72 @@ float3 GetLightUV(int LightId, RayDesc Ray, FMaterialClosestHitPayload RayHitInf
 
 float3 GetPosition(int LightId)
 {
-	return SceneLightsData.Position[LightId];
+	return SceneLights[LightId].Position;
 }
 
 float3 GetNormal(int LightId)
 {
-	return SceneLightsData.Normal[LightId];
+	return SceneLights[LightId].Normal;
 }
 
 float3 GetdPdu(int LightId)
 {
-	return SceneLightsData.dPdu[LightId];
+	return SceneLights[LightId].dPdu;
 }
 
 float3 GetdPdv(int LightId)
 {
-	return SceneLightsData.dPdv[LightId];
+	return SceneLights[LightId].dPdv;
 }
 
 float GetWidth(int LightId)
 {
-	return SceneLightsData.Dimensions[LightId].x;
+	return SceneLights[LightId].Dimensions.x;
 }
 
 float GetHeight(int LightId)
 {
-	return SceneLightsData.Dimensions[LightId].y;
+	return SceneLights[LightId].Dimensions.y;
 }
 
 float2 GetRectSize(int LightId)
 {
-	return SceneLightsData.Dimensions[LightId].xy;
+	return SceneLights[LightId].Dimensions.xy;
 }
 
 float GetDepth(int LightId)
 {
-	return SceneLightsData.Dimensions[LightId].z;
+	return SceneLights[LightId].Dimensions.z;
 }
 
 float2 GetCosConeAngles(int LightId)
 {
-	return SceneLightsData.Dimensions[LightId].xy;
+	return SceneLights[LightId].Dimensions.xy;
 }
 
 float3 GetColor(int LightId)
 {
-	return SceneLightsData.Color[LightId];
+	return SceneLights[LightId].Color;
 }
 
 float GetRadius(int LightId)
 {
-	return SceneLightsData.Dimensions[LightId].z;
+	return SceneLights[LightId].Dimensions.z;
 }
 
 float GetAttenuation(int LightId)
 {
-	return SceneLightsData.Attenuation[LightId];
+	return SceneLights[LightId].Attenuation;
 }
 
 float GetRectLightBarnCosAngle(int LightId)
 {
-	return SceneLightsData.RectLightBarnCosAngle[LightId];
+	return SceneLights[LightId].RectLightBarnCosAngle;
 }
 
 float GetRectLightBarnLength(int LightId)
 {
-	return SceneLightsData.RectLightBarnLength[LightId];
+	return SceneLights[LightId].RectLightBarnLength;
 }
 
 #ifndef USE_ATTENUATION_TERM
@@ -133,7 +139,7 @@ float ComputeAttenuationFalloff(float DistanceSquared, int LightId)
 	// #dxr_todo: UE-72508: encapsulate this function in a shared space
 	float InvAttenuationRadius = GetAttenuation(LightId);
 	float NormalizeDistanceSquared = DistanceSquared * Square(InvAttenuationRadius);
-	if ((SceneLightsData.Flags[LightId] & PATHTRACER_FLAG_NON_INVERSE_SQUARE_FALLOFF_MASK) == 0)
+	if ((SceneLights[LightId].Flags & PATHTRACER_FLAG_NON_INVERSE_SQUARE_FALLOFF_MASK) == 0)
 	{
 		return Square(saturate(1.0 - Square(NormalizeDistanceSquared)));
 	}
@@ -141,7 +147,7 @@ float ComputeAttenuationFalloff(float DistanceSquared, int LightId)
 	{
 		// roughly cancel out "native" square distance falloff before applying exponent based falloff function
         // this appears to match the behavior of the deferred lighting passes
-		float FalloffExponent = SceneLightsData.FalloffExponent[LightId];
+		float FalloffExponent = SceneLights[LightId].FalloffExponent;
 		return DistanceSquared * pow(1.0 - saturate(NormalizeDistanceSquared), FalloffExponent);
 	}
 #else
@@ -163,7 +169,7 @@ float ComputeIESAttenuation(int LightId, float3 WorldPosition)
 {
 	float Result = 1.f;
 #if USE_IES_TERM
-	if (SceneLightsData.IESTextureSlice[LightId] >= 0)
+	if (SceneLights[LightId].IESTextureSlice >= 0)
 	{
 		float3 LightDirection = normalize(GetPosition(LightId) - WorldPosition);
 
@@ -175,7 +181,7 @@ float ComputeIESAttenuation(int LightId, float3 WorldPosition)
 		float dv = dot(LightDirection, GetdPdv(LightId));
 		float NormTangentAngle = atan2(dv, du) / (PI * 2.f) + 0.5f;
 
-		const float3 UVW = float3(NormAngle, NormTangentAngle, SceneLightsData.IESTextureSlice[LightId]);
+		const float3 UVW = float3(NormAngle, NormTangentAngle, SceneLights[LightId].IESTextureSlice);
 		Result = IESTexture.SampleLevel(IESTextureSampler, UVW, 0);
 	}
 #endif
@@ -184,15 +190,15 @@ float ComputeIESAttenuation(int LightId, float3 WorldPosition)
 
 bool HasTransmission(int LightId)
 {
-	return (SceneLightsData.Flags[LightId] & PATHTRACER_FLAG_TRANSMISSION_MASK) != 0;
+	return (SceneLights[LightId].Flags & PATHTRACER_FLAG_TRANSMISSION_MASK) != 0;
 }
 
 uint GetLightingChannelMask(int LightId)
 {
-	return SceneLightsData.Flags[LightId] & PATHTRACER_FLAG_LIGHTING_CHANNEL_MASK;
+	return SceneLights[LightId].Flags & PATHTRACER_FLAG_LIGHTING_CHANNEL_MASK;
 }
 
 bool IsStationary(int LightId)
 {
-	return (SceneLightsData.Flags[LightId] & PATHTRACER_FLAG_STATIONARY_MASK) != 0;
+	return (SceneLights[LightId].Flags & PATHTRACER_FLAG_STATIONARY_MASK) != 0;
 }
diff --git a/Engine/Shaders/Private/PathTracing/Light/PathTracingLightSampling.ush b/Engine/Shaders/Private/PathTracing/Light/PathTracingLightSampling.ush
--- a/Engine/Shaders/Private/PathTracing/Light/PathTracingLightSampling.ush
+++ b/Engine/Shaders/Private/PathTracing/Light/PathTracingLightSampling.ush
@@ -22,7 +22,7 @@ LightSampling.usf: Light sampling for next-event estimation
 // Returns float4 packed with Radiance and Pdf
 float4 TraceLight(RayDesc Ray, int LightId, float HitT)
 {
-	switch (SceneLightsData.Flags[LightId] & PATHTRACER_FLAG_TYPE_MASK)
+	switch (SceneLights[LightId].Flags & PATHTRACER_FLAG_TYPE_MASK)
 	{
 		case PATHTRACING_LIGHT_SKY:			return SkyLight_TraceLight(Ray, LightId, HitT);
 		case PATHTRACING_LIGHT_POINT:		return PointLight_TraceLight(Ray, LightId, HitT);
@@ -44,7 +44,7 @@ void SampleLight(
 	out float3 OutRadianceOverPdf,
 	out float OutPdf)
 {
-	switch (SceneLightsData.Flags[LightId] & PATHTRACER_FLAG_TYPE_MASK)
+	switch (SceneLights[LightId].Flags & PATHTRACER_FLAG_TYPE_MASK)
 	{
 		case PATHTRACING_LIGHT_SKY:
 			SkyLight_SampleLight(LightId, RandSample, WorldPos, WorldNormal, OutDirection, OutDistance, OutRadianceOverPdf, OutPdf);
@@ -83,7 +83,7 @@ void EstimateLight(
 	OutIrradiance = 0.0f;
 	if ((GetLightingChannelMask(LightId) & PrimitiveLightingChannelMask) == 0) return;
 
-	switch (SceneLightsData.Flags[LightId] & PATHTRACER_FLAG_TYPE_MASK)
+	switch (SceneLights[LightId].Flags & PATHTRACER_FLAG_TYPE_MASK)
 	{
 		case PATHTRACING_LIGHT_SKY:
 			SkyLight_EstimateLight(LightId, WorldPos, WorldNormal, IsTransmissiveMaterial, OutIrradiance);
@@ -108,9 +108,9 @@ void EstimateLight(
 // #dxr_todo: passing LightPickingCdf as anything but inout causes performance to drop significantly -- why ?
 void InitLightPickingCdf(inout float LightPickingCdf[RAY_TRACING_LIGHT_COUNT_MAXIMUM])
 {
-	for (uint LightIndex = 0; LightIndex < SceneLightsData.Count; ++LightIndex)
+	for (uint LightIndex = 0; LightIndex < SceneLightCount; ++LightIndex)
 	{
-		LightPickingCdf[LightIndex] = float(LightIndex + 1) / float(SceneLightsData.Count);
+		LightPickingCdf[LightIndex] = float(LightIndex + 1) / float(SceneLightCount);
 	}
 }
 
@@ -122,14 +122,14 @@ bool InitLightPickingCdf(
 	bool IsTransmissiveMaterial,
 	inout float LightPickingCdf[RAY_TRACING_LIGHT_COUNT_MAXIMUM])
 {
-	// This code assumes SceneLightsData.Count > 0
+	// This code assumes SceneLightCount > 0
 #ifdef UNIFORM_LIGHT_SELECTION
 	InitLightPickingCdf(LightPickingCdf);
 #else
 	// Build irradiance estimate prefix sum
 	float LastCdf = 0;
 	{
-		for (uint LightIndex = 0; LightIndex < SceneLightsData.Count; ++LightIndex)
+		for (uint LightIndex = 0; LightIndex < SceneLightCount; ++LightIndex)
 		{
 			EstimateLight(LightIndex, WorldPos, WorldNormal, PrimitiveLightingChannelMask, IsTransmissiveMaterial, LightPickingCdf[LightIndex]);
 			LightPickingCdf[LightIndex] += LastCdf;
@@ -144,11 +144,11 @@ bool InitLightPickingCdf(
 	}
 	{
 		// CDF normalization
-		for (uint LightIndex = 0; LightIndex < SceneLightsData.Count; ++LightIndex)
+		for (uint LightIndex = 0; LightIndex < SceneLightCount; ++LightIndex)
 		{
 			LightPickingCdf[LightIndex] /= LastCdf;
 		}
-		LightPickingCdf[SceneLightsData.Count - 1] = 1.0;
+		LightPickingCdf[SceneLightCount - 1] = 1.0;
 	}
 #endif
 	return true;
@@ -157,12 +157,12 @@ bool InitLightPickingCdf(
 // #dxr_todo: passing LightPickingCdf as anything but inout causes performance to drop significantly -- why ?
 void SelectLight(float RandSample, inout float LightPickingCdf[RAY_TRACING_LIGHT_COUNT_MAXIMUM], out uint LightId, out float LightPickPdf)
 {
-	// This code assumes SceneLightsData.Count > 0
+	// This code assumes SceneLightCount > 0
 
 	// Draw random variable
 	// #dxr_todo: binary search
 	float PreviousCdfValue = 0;
-	for (LightId = 0; LightId < SceneLightsData.Count - 1; ++LightId)
+	for (LightId = 0; LightId < SceneLightCount - 1; ++LightId)
 	{
 		if (RandSample < LightPickingCdf[LightId])
 		{
@@ -193,7 +193,7 @@ bool SampleLight(
 	OutRadianceOverPdf = 0;
 	OutPdf = 0;
 
-	if (SceneLightsData.Count == 0)
+	if (SceneLightCount == 0)
 	{
 		// no lights in the scene
 		return false;
diff --git a/Engine/Shaders/Private/PathTracing/PathTracing.usf b/Engine/Shaders/Private/PathTracing/PathTracing.usf
--- a/Engine/Shaders/Private/PathTracing/PathTracing.usf
+++ b/Engine/Shaders/Private/PathTracing/PathTracing.usf
@@ -91,6 +91,7 @@ RAY_TRACING_ENTRY_RAYGEN(PathTracingMainRG)
 	// We don't want the camera ray to have any MIS contribution against direct light sampling so we pretend that it came from
 	// a singular Material scattering event.
 	float PrevMaterialPdf = POSITIVE_INFINITY;
+	int PrevPrimitiveLightingChannelMask = 0xFF;
 
 	// 0: only Material sampling
 	// 1: only Light sampling
@@ -116,12 +117,20 @@ RAY_TRACING_ENTRY_RAYGEN(PathTracingMainRG)
 			// Loop over lights to capture their contribution
 			// #dxr_todo: if we have lots of lights, having some hierarchical structure would be better ....
 			float PreviousCdfValue = 0.0;
-			for (uint LightId = 0; LightId < SceneLightsData.Count; ++LightId)
+			for (uint LightId = 0; LightId < SceneLightCount; ++LightId)
 			{
 				// Camera rays should only do the rest of the loop if we have visible lights OR if
 				// we have a skydome (which are always visible)
 				if (Bounce == 0 && PathTracingData.VisibleLights == 0 && !IsEnvironmentLight(LightId))
+				{
 					continue;
+				}
+
+				if ((PrevPrimitiveLightingChannelMask & GetLightingChannelMask(LightId)) == 0)
+				{
+					// light does not affect the current ray
+					continue;
+				}
 
 				float4 LightResult = TraceLight(Ray, LightId, HitT);
 				float3 LightRadiance = LightResult.xyz;
@@ -159,7 +168,7 @@ RAY_TRACING_ENTRY_RAYGEN(PathTracingMainRG)
 		bool bIsNeeValid = IsNeeValidMaterial(Payload);
 
 		// If we are using Light sampling and we are not on the last bounce ....
-		if (MISMode != 0 && Bounce < PathTracingData.MaxBounces && bIsNeeValid && SceneLightsData.Count > 0)
+		if (MISMode != 0 && Bounce < PathTracingData.MaxBounces && bIsNeeValid && SceneLightCount > 0)
 		{
 			// Choose a light and sample it
 
@@ -193,7 +202,7 @@ RAY_TRACING_ENTRY_RAYGEN(PathTracingMainRG)
 				if (LightPdf > 0)
 				{
 					float SignedPositionBias = 1.0;
-					if (ENABLE_TRANSMISSION && IsMaterialTransmissive(Payload))
+					if (IsTransmissiveMaterial)
 					{
 						// for transmissive materials, bias the position to the other side of the surface if the light is coming from behind
 						float NoL = dot(Payload.WorldNormal, LightRay.Direction);
@@ -295,6 +304,8 @@ RAY_TRACING_ENTRY_RAYGEN(PathTracingMainRG)
 			PrevMaterialPdf = MaterialSamplePdf;
 		}
 
+		PrevPrimitiveLightingChannelMask = Payload.PrimitiveLightingChannelMask;
+
 		// Update ray and loop back around to trace again
 		Ray.Direction = Direction;
 		Ray.TMin = 0;
diff --git a/Engine/Shaders/Shared/RayTracingDefinitions.h b/Engine/Shaders/Shared/RayTracingDefinitions.h
--- a/Engine/Shaders/Shared/RayTracingDefinitions.h
+++ b/Engine/Shaders/Shared/RayTracingDefinitions.h
@@ -1,7 +1,7 @@
 // Copyright Epic Games, Inc. All Rights Reserved.
 
 /*================================================================================================
-	RayTracingDefinitions.ush: used in ray tracing shaders and C++ code to define common constants
+	RayTracingDefinitions.h: used in ray tracing shaders and C++ code to define common constants
 	!!! Changing this file requires recompilation of the engine !!!
 =================================================================================================*/
 
@@ -35,7 +35,7 @@
 #define RAY_TRACING_MAX_ALLOWED_ATTRIBUTE_SIZE  8   // Sizeof 2 floats (barycentrics)
 #define RAY_TRACING_MAX_ALLOWED_PAYLOAD_SIZE    64  // Our maximum allowed payload size (sizeof FPackedMaterialClosestHitPayload)
 
-// Constants for the 'Flags' field of FPathTracingLightData
+// Constants for the 'Flags' field of FPathTracingLight
 #define PATHTRACER_FLAG_LIGHTING_CHANNEL_MASK				(7 << 0)	// Which lighting channel is this light assigned to?
 #define PATHTRACER_FLAG_TRANSMISSION_MASK					(1 << 3)	// Does the light affect the transmission side?
 #define PATHTRACER_FLAG_NON_INVERSE_SQUARE_FALLOFF_MASK		(1 << 4)	// Does the light have a non-inverse square decay?
diff --git a/Engine/Shaders/Shared/RayTracingTypes.h b/Engine/Shaders/Shared/RayTracingTypes.h
new file mode 100644
--- /dev/null
+++ b/Engine/Shaders/Shared/RayTracingTypes.h
@@ -0,0 +1,52 @@
+// Copyright Epic Games, Inc. All Rights Reserved.
+
+/*================================================================================================
+	RayTracingTypes.h: used in ray tracing shaders and C++ code to define common types
+	!!! Changing this file requires recompilation of the engine !!!
+=================================================================================================*/
+
+#ifdef __cplusplus
+
+// C++ representation of a light for the path tracer
+// #dxr_todo: Unify this with FRTLightingData ?
+struct FPathTracingLight {
+	FVector Position;
+	FVector Normal;
+	FVector dPdu;
+	FVector dPdv;
+	FVector Color;
+	FVector Dimensions; // Radius,SoftRadius,Length or RectWidth,RectHeight or Sin(Angle/2),Sin(SoftAngle/2) depending on light type
+	float   Attenuation;
+	float   FalloffExponent; // for non-inverse square decay lights only
+	float   RectLightBarnCosAngle;
+	float   RectLightBarnLength;
+	int32   IESTextureSlice;
+	uint32  Flags; // see defines above
+
+	float Padding[8]; // keep structure aligned - take new fields from here as needed
+};
+
+static_assert(sizeof(FPathTracingLight) == 128, "Path tracing light structure should be aligned to 128 bytes for optimal access on the GPU");
+
+#else
+
+// HLSL side of the struct above
+
+struct FPathTracingLight {
+	float3  Position;
+	float3  Normal;
+	float3  dPdu;
+	float3  dPdv;
+	float3  Color;
+	float3  Dimensions;
+	float   Attenuation;
+	float   FalloffExponent;
+	float   RectLightBarnCosAngle;
+	float   RectLightBarnLength;
+	int     IESTextureSlice;
+	uint    Flags;
+
+	float Padding[8];
+};
+
+#endif
diff --git a/Engine/Source/Runtime/Renderer/Private/PathTracing.cpp b/Engine/Source/Runtime/Renderer/Private/PathTracing.cpp
--- a/Engine/Source/Runtime/Renderer/Private/PathTracing.cpp
+++ b/Engine/Source/Runtime/Renderer/Private/PathTracing.cpp
@@ -14,6 +14,7 @@
 #include "RayTracing/RayTracingSkyLight.h"
 #include "RayTracing/RaytracingOptions.h"
 #include "HAL/PlatformApplicationMisc.h"
+#include "RayTracingTypes.h"
 
 TAutoConsoleVariable<int32> CVarPathTracingMaxBounces(
 	TEXT("r.PathTracing.MaxBounces"),
@@ -96,7 +97,6 @@ TAutoConsoleVariable<int32> CVarPathTracingProgressDisplay(
 );
 
 IMPLEMENT_GLOBAL_SHADER_PARAMETER_STRUCT(FPathTracingData, "PathTracingData");
-IMPLEMENT_GLOBAL_SHADER_PARAMETER_STRUCT(FPathTracingLightData, "SceneLightsData");
 
 // This function prepares the portion of shader arguments that may involve invalidating the path traced state
 static bool PrepareShaderArgs(const FViewInfo& View, FPathTracingData& PathTracingData) {
@@ -174,8 +174,9 @@ class FPathTracingRG : public FGlobalShader
 
 		SHADER_PARAMETER_STRUCT_REF(FViewUniformShaderParameters, ViewUniformBuffer)
 		SHADER_PARAMETER_STRUCT_REF(FSkyLightData, SkyLightData)
-		SHADER_PARAMETER_STRUCT_REF(FPathTracingLightData, SceneLightsData)
 		SHADER_PARAMETER_STRUCT_REF(FPathTracingData, PathTracingData)
+		SHADER_PARAMETER_RDG_BUFFER_SRV(StructuredBuffer<FPathTracingLight>, SceneLights)
+		SHADER_PARAMETER(uint32, SceneLightCount)
 		// IES Profiles
 		SHADER_PARAMETER_RDG_TEXTURE(Texture2DArray, IESTexture)
 		SHADER_PARAMETER_SAMPLER(SamplerState, IESTextureSampler)
@@ -211,40 +212,42 @@ class FPathTracingIESAtlasCS : public FGlobalShader
 };
 IMPLEMENT_SHADER_TYPE(, FPathTracingIESAtlasCS, TEXT("/Engine/Private/PathTracing/PathTracingIESAtlas.usf"), TEXT("PathTracingIESAtlasCS"), SF_Compute);
 
-void SetLightParameters(FRDGBuilder& GraphBuilder, FPathTracingRG::FParameters* PassParameters, FPathTracingLightData& LightData, FSkyLightData& SkyLightData, FScene* Scene, const FViewInfo& View, bool UseLightProfiles)
+void SetLightParameters(FRDGBuilder& GraphBuilder, FPathTracingRG::FParameters* PassParameters, FSkyLightData& SkyLightData, FScene* Scene, const FViewInfo& View, bool UseLightProfiles)
 {
 	// Sky light
 	bool IsSkyLightValid = SetupSkyLightParameters(*Scene, &SkyLightData);
 
 	// Lights
-	LightData.Count = 0;
+	FPathTracingLight Lights[RAY_TRACING_LIGHT_COUNT_MAXIMUM]; // Keep this on the stack for now -- eventually will need to make this dynamic to lift size limit (and also avoid uploading per frame ...)
+	unsigned LightCount = 0;
 
 	// Prepend SkyLight to light buffer since it is not part of the regular light list
 	if (IsSkyLightValid)
 	{
-		uint32 SkyLightIndex = 0;
-		uint8 SkyLightLightingChannelMask = 0xFF;
-		LightData.Color[SkyLightIndex] = FVector(SkyLightData.Color);
-		LightData.Flags[SkyLightIndex] = SkyLightData.bTransmission ? PATHTRACER_FLAG_TRANSMISSION_MASK : 0;
-		LightData.Flags[SkyLightIndex] |= PATHTRACER_FLAG_LIGHTING_CHANNEL_MASK;
-		LightData.Flags[SkyLightIndex] |= PATHTRACING_LIGHT_SKY;
-		LightData.IESTextureSlice[SkyLightIndex] = -1;
-		LightData.FalloffExponent[LightData.Count] = 0;
-		LightData.Count++;
+		FPathTracingLight& DestLight = Lights[LightCount];
+		DestLight.Color = FVector(SkyLightData.Color);
+		DestLight.Flags = SkyLightData.bTransmission ? PATHTRACER_FLAG_TRANSMISSION_MASK : 0;
+		DestLight.Flags |= PATHTRACER_FLAG_LIGHTING_CHANNEL_MASK;
+		DestLight.Flags |= PATHTRACING_LIGHT_SKY;
+		DestLight.IESTextureSlice = -1;
+		LightCount++;
 	}
 
 	TMap<FTexture*, int> IESLightProfilesMap;
 	for (auto Light : Scene->Lights)
 	{
-		if (LightData.Count >= RAY_TRACING_LIGHT_COUNT_MAXIMUM) break;
+		if (LightCount >= RAY_TRACING_LIGHT_COUNT_MAXIMUM) break;
+
+		FPathTracingLight& DestLight = Lights[LightCount];
 
 		FLightShaderParameters LightParameters;
 		Light.LightSceneInfo->Proxy->GetLightShaderParameters(LightParameters);
 		uint32 Transmission = Light.LightSceneInfo->Proxy->Transmission();
 		uint8 LightingChannelMask = Light.LightSceneInfo->Proxy->GetLightingChannelMask();
-		LightData.Flags[LightData.Count] = Transmission ? PATHTRACER_FLAG_TRANSMISSION_MASK : 0;
-		LightData.Flags[LightData.Count] |= LightingChannelMask & PATHTRACER_FLAG_LIGHTING_CHANNEL_MASK;
-		LightData.IESTextureSlice[LightData.Count] = -1;
+
+		DestLight.Flags = Transmission ? PATHTRACER_FLAG_TRANSMISSION_MASK : 0;
+		DestLight.Flags |= LightingChannelMask & PATHTRACER_FLAG_LIGHTING_CHANNEL_MASK;
+		DestLight.IESTextureSlice = -1;
 
 		if (UseLightProfiles)
 		{
@@ -252,53 +255,53 @@ void SetLightParameters(FRDGBuilder& GraphBuilder, FPathTracingRG::FParameters*
 			if (IESTexture != nullptr)
 			{
 				// Only add a given texture once
-				LightData.IESTextureSlice[LightData.Count] = IESLightProfilesMap.FindOrAdd(IESTexture, IESLightProfilesMap.Num());
+				DestLight.IESTextureSlice = IESLightProfilesMap.FindOrAdd(IESTexture, IESLightProfilesMap.Num());
 			}
 		}
 
 		// these mean roughly the same thing across all light types
-		LightData.Color[LightData.Count] = LightParameters.Color;
-		LightData.Position[LightData.Count] = LightParameters.Position;
-		LightData.Normal[LightData.Count] = -LightParameters.Direction;
-		LightData.dPdu[LightData.Count] = FVector::CrossProduct(LightParameters.Tangent, LightParameters.Direction);
-		LightData.dPdv[LightData.Count] = LightParameters.Tangent;
-		LightData.Attenuation[LightData.Count] = LightParameters.InvRadius;
-		LightData.FalloffExponent[LightData.Count] = 0;
+		DestLight.Color = LightParameters.Color;
+		DestLight.Position = LightParameters.Position;
+		DestLight.Normal = -LightParameters.Direction;
+		DestLight.dPdu = FVector::CrossProduct(LightParameters.Tangent, LightParameters.Direction);
+		DestLight.dPdv = LightParameters.Tangent;
+		DestLight.Attenuation = LightParameters.InvRadius;
+		DestLight.FalloffExponent = 0;
 
 		ELightComponentType LightComponentType = (ELightComponentType)Light.LightSceneInfo->Proxy->GetLightType();
 		switch (LightComponentType)
 		{
 			case LightType_Directional:
 			{
-                LightData.Normal[LightData.Count] = LightParameters.Direction;
-				LightData.Dimensions[LightData.Count] = FVector(0.0f, 0.0f, LightParameters.SourceRadius);
-				LightData.Flags[LightData.Count] |= PATHTRACING_LIGHT_DIRECTIONAL;
+				DestLight.Normal = LightParameters.Direction;
+				DestLight.Dimensions = FVector(0.0f, 0.0f, LightParameters.SourceRadius);
+				DestLight.Flags |= PATHTRACING_LIGHT_DIRECTIONAL;
 				break;
 			}
 			case LightType_Rect:
 			{
-				LightData.Dimensions[LightData.Count] = FVector(2.0f * LightParameters.SourceRadius, 2.0f * LightParameters.SourceLength, 0.0f);
-				LightData.RectLightBarnCosAngle[LightData.Count] = LightParameters.RectLightBarnCosAngle;
-				LightData.RectLightBarnLength[LightData.Count] = LightParameters.RectLightBarnLength;
-				LightData.FalloffExponent[LightData.Count] = LightParameters.FalloffExponent;
-				LightData.Flags[LightData.Count] |= Light.LightSceneInfo->Proxy->IsInverseSquared() ? 0 : PATHTRACER_FLAG_NON_INVERSE_SQUARE_FALLOFF_MASK;
-				LightData.Flags[LightData.Count] |= PATHTRACING_LIGHT_RECT;
+				DestLight.Dimensions = FVector(2.0f * LightParameters.SourceRadius, 2.0f * LightParameters.SourceLength, 0.0f);
+				DestLight.RectLightBarnCosAngle = LightParameters.RectLightBarnCosAngle;
+				DestLight.RectLightBarnLength = LightParameters.RectLightBarnLength;
+				DestLight.FalloffExponent = LightParameters.FalloffExponent;
+				DestLight.Flags |= Light.LightSceneInfo->Proxy->IsInverseSquared() ? 0 : PATHTRACER_FLAG_NON_INVERSE_SQUARE_FALLOFF_MASK;
+				DestLight.Flags |= PATHTRACING_LIGHT_RECT;
 				break;
 			}
 			case LightType_Spot:
 			{
-				LightData.Dimensions[LightData.Count] = FVector(LightParameters.SpotAngles, LightParameters.SourceRadius);
-				LightData.FalloffExponent[LightData.Count] = LightParameters.FalloffExponent;
-				LightData.Flags[LightData.Count] |= Light.LightSceneInfo->Proxy->IsInverseSquared() ? 0 : PATHTRACER_FLAG_NON_INVERSE_SQUARE_FALLOFF_MASK;
-				LightData.Flags[LightData.Count] |= PATHTRACING_LIGHT_SPOT;
+				DestLight.Dimensions = FVector(LightParameters.SpotAngles, LightParameters.SourceRadius);
+				DestLight.FalloffExponent = LightParameters.FalloffExponent;
+				DestLight.Flags |= Light.LightSceneInfo->Proxy->IsInverseSquared() ? 0 : PATHTRACER_FLAG_NON_INVERSE_SQUARE_FALLOFF_MASK;
+				DestLight.Flags |= PATHTRACING_LIGHT_SPOT;
 				break;
 			}
 			case LightType_Point:
 			{
-				LightData.Dimensions[LightData.Count] = FVector(0.0, 0.0, LightParameters.SourceRadius);
-				LightData.FalloffExponent[LightData.Count] = LightParameters.FalloffExponent;
-				LightData.Flags[LightData.Count] |= Light.LightSceneInfo->Proxy->IsInverseSquared() ? 0 : PATHTRACER_FLAG_NON_INVERSE_SQUARE_FALLOFF_MASK;
-				LightData.Flags[LightData.Count] |= PATHTRACING_LIGHT_POINT;
+				DestLight.Dimensions = FVector(0.0, 0.0, LightParameters.SourceRadius);
+				DestLight.FalloffExponent = LightParameters.FalloffExponent;
+				DestLight.Flags |= Light.LightSceneInfo->Proxy->IsInverseSquared() ? 0 : PATHTRACER_FLAG_NON_INVERSE_SQUARE_FALLOFF_MASK;
+				DestLight.Flags |= PATHTRACING_LIGHT_POINT;
 				break;
 			}
 			default:
@@ -307,9 +310,16 @@ void SetLightParameters(FRDGBuilder& GraphBuilder, FPathTracingRG::FParameters*
 				checkNoEntry();
 				break;
 			}
-		};
+		}
 
-		LightData.Count++;
+		LightCount++;
+	}
+
+	{
+		// Upload the buffer of lights to the GPU
+		size_t DataSize = sizeof(FPathTracingLight) * FMath::Max(LightCount, 1u);
+		PassParameters->SceneLights = GraphBuilder.CreateSRV(FRDGBufferSRVDesc(CreateStructuredBuffer(GraphBuilder, TEXT("PathTracingLightsBuffer"), sizeof(FPathTracingLight), FMath::Max(LightCount, 1u), Lights, DataSize)));
+		PassParameters->SceneLightCount = LightCount;
 	}
 
 	if (!IESLightProfilesMap.IsEmpty())
@@ -375,8 +385,8 @@ class FPathTracingCompositorPS : public FGlobalShader
 	BEGIN_SHADER_PARAMETER_STRUCT(FParameters, )
 		SHADER_PARAMETER_RDG_TEXTURE_SRV(Texture2D<float4>, RadianceTexture)
 		SHADER_PARAMETER_STRUCT_REF(FViewUniformShaderParameters, ViewUniformBuffer)
-		SHADER_PARAMETER(unsigned, Iteration)
-		SHADER_PARAMETER(unsigned, MaxSamples)
+		SHADER_PARAMETER(uint32, Iteration)
+		SHADER_PARAMETER(uint32, MaxSamples)
 		SHADER_PARAMETER(int, ProgressDisplayEnabled)
 
 		RENDER_TARGET_BINDING_SLOTS()
@@ -496,10 +506,8 @@ void FDeferredShadingSceneRenderer::RenderPathTracing(
 		{
 			// upload sky/lights data
 			FSkyLightData SkyLightData;
-			FPathTracingLightData LightData;
-			SetLightParameters(GraphBuilder, PassParameters, LightData, SkyLightData, Scene, View, UseLightProfiles);
+			SetLightParameters(GraphBuilder, PassParameters, SkyLightData, Scene, View, UseLightProfiles);
 			PassParameters->SkyLightData = CreateUniformBufferImmediate(SkyLightData, EUniformBufferUsage::UniformBuffer_SingleFrame);
-			PassParameters->SceneLightsData = CreateUniformBufferImmediate(LightData, EUniformBufferUsage::UniformBuffer_SingleFrame);
 		}
 		PassParameters->IESTextureSampler = TStaticSamplerState<SF_Bilinear, AM_Clamp, AM_Clamp, AM_Clamp>::GetRHI();
 		PassParameters->RadianceTexture = GraphBuilder.CreateUAV(RadianceTexture);
diff --git a/Engine/Source/Runtime/Renderer/Private/PathTracingUniformBuffers.h b/Engine/Source/Runtime/Renderer/Private/PathTracingUniformBuffers.h
--- a/Engine/Source/Runtime/Renderer/Private/PathTracingUniformBuffers.h
+++ b/Engine/Source/Runtime/Renderer/Private/PathTracingUniformBuffers.h
@@ -20,27 +20,3 @@ BEGIN_GLOBAL_SHADER_PARAMETER_STRUCT(FPathTracingData, )
 	SHADER_PARAMETER(float, MaxPathIntensity)
 	SHADER_PARAMETER(float, MaxNormalBias)
 END_GLOBAL_SHADER_PARAMETER_STRUCT()
-
-
-// Lights
-
-BEGIN_GLOBAL_SHADER_PARAMETER_STRUCT(FPathTracingLightData, RENDERER_API)
-	SHADER_PARAMETER(uint32, Count)
-	// Geometry
-	SHADER_PARAMETER_ARRAY(FVector, Position, [RAY_TRACING_LIGHT_COUNT_MAXIMUM])
-	SHADER_PARAMETER_ARRAY(FVector, Normal, [RAY_TRACING_LIGHT_COUNT_MAXIMUM])
-	SHADER_PARAMETER_ARRAY(FVector, dPdu, [RAY_TRACING_LIGHT_COUNT_MAXIMUM])
-	SHADER_PARAMETER_ARRAY(FVector, dPdv, [RAY_TRACING_LIGHT_COUNT_MAXIMUM])
-	// Color
-	SHADER_PARAMETER_ARRAY(FVector, Color, [RAY_TRACING_LIGHT_COUNT_MAXIMUM])
-	// Light-specific
-	SHADER_PARAMETER_ARRAY(FVector, Dimensions, [RAY_TRACING_LIGHT_COUNT_MAXIMUM])
-	SHADER_PARAMETER_ARRAY(float, Attenuation, [RAY_TRACING_LIGHT_COUNT_MAXIMUM])
-	SHADER_PARAMETER_ARRAY(float, FalloffExponent, [RAY_TRACING_LIGHT_COUNT_MAXIMUM])
-	SHADER_PARAMETER_ARRAY(float, RectLightBarnCosAngle, [RAY_TRACING_LIGHT_COUNT_MAXIMUM])
-	SHADER_PARAMETER_ARRAY(float, RectLightBarnLength, [RAY_TRACING_LIGHT_COUNT_MAXIMUM])
-	// IES index (valid if >= 0)
-	SHADER_PARAMETER_ARRAY(float, IESTextureSlice , [RAY_TRACING_LIGHT_COUNT_MAXIMUM])
-	// Flags
-	SHADER_PARAMETER_ARRAY(uint32, Flags, [RAY_TRACING_LIGHT_COUNT_MAXIMUM])
-END_GLOBAL_SHADER_PARAMETER_STRUCT()
diff --git a/Engine/Source/Runtime/Renderer/Private/RayTracing/RayTracingGlobalIllumination.cpp b/Engine/Source/Runtime/Renderer/Private/RayTracing/RayTracingGlobalIllumination.cpp
--- a/Engine/Source/Runtime/Renderer/Private/RayTracing/RayTracingGlobalIllumination.cpp
+++ b/Engine/Source/Runtime/Renderer/Private/RayTracing/RayTracingGlobalIllumination.cpp
@@ -23,6 +23,7 @@
 #include "SceneTextureParameters.h"
 #include "RayTracingDefinitions.h"
 #include "RayTracingDeferredMaterials.h"
+#include "RayTracingTypes.h"
 
 static TAutoConsoleVariable<int32> CVarRayTracingGlobalIllumination(
 	TEXT("r.RayTracing.GlobalIllumination"),
@@ -245,10 +246,10 @@ DECLARE_GPU_STAT_NAMED(RayTracingGICreateGatherPoints, TEXT("Ray Tracing GI: Cre
 
 RENDERER_API void SetupLightParameters(
 	const FScene& Scene,
-	const FViewInfo& View,
-	FPathTracingLightData* LightParameters)
+	const FViewInfo& View, FRDGBuilder& GraphBuilder, FRDGBufferSRV** OutLightBuffer, uint32* OutLightCount)
 {
-	LightParameters->Count = 0;
+	FPathTracingLight Lights[RAY_TRACING_LIGHT_COUNT_MAXIMUM];
+	unsigned LightCount = 0;
 
 	// Get the SkyLight color
 
@@ -261,37 +262,42 @@ RENDERER_API void SetupLightParameters(
 		SkyLightTransmission = SkyLight->bTransmission;
 	}
 
-	// Prepend SkyLight to light buffer
-	// WARNING: Until ray payload encodes Light data buffer, the execution depends on this ordering!
-	uint32 SkyLightIndex = 0;
-	LightParameters->Color[SkyLightIndex] = SkyLightColor;
-	LightParameters->Flags[SkyLightIndex] = SkyLightTransmission ? PATHTRACER_FLAG_TRANSMISSION_MASK : 0;
-	// SkyLight does not have a LightingChannelMask
-	LightParameters->Flags[SkyLightIndex] |= PATHTRACER_FLAG_LIGHTING_CHANNEL_MASK;
-	LightParameters->Flags[SkyLightIndex] |= PATHTRACING_LIGHT_SKY;
-	
-	LightParameters->Count++;
+	// Prepend SkyLight to light buffer (if it is active)
+	{
+		FPathTracingLight& DestLight = Lights[LightCount];
+
+		DestLight.Color = SkyLightColor;
+		DestLight.Flags = SkyLightTransmission ? PATHTRACER_FLAG_TRANSMISSION_MASK : 0;
+		// SkyLight does not have a LightingChannelMask
+		DestLight.Flags |= PATHTRACER_FLAG_LIGHTING_CHANNEL_MASK;
+		DestLight.Flags |= PATHTRACING_LIGHT_SKY;
 
+		LightCount++;
+	}
+
+	
 	uint32 MaxLightCount = FMath::Min(CVarRayTracingGlobalIlluminationMaxLightCount.GetValueOnRenderThread(), RAY_TRACING_LIGHT_COUNT_MAXIMUM);
 	for (auto Light : Scene.Lights)
 	{
-		if (LightParameters->Count >= MaxLightCount) break;
+		if (LightCount >= MaxLightCount) break;
 
 		if (Light.LightSceneInfo->Proxy->HasStaticLighting() && Light.LightSceneInfo->IsPrecomputedLightingValid()) continue;
 		if (!Light.LightSceneInfo->Proxy->AffectGlobalIllumination()) continue;
 
+		FPathTracingLight& DestLight = Lights[LightCount]; // don't increment LightCount yet -- we might still skip this light
+
 		FLightShaderParameters LightShaderParameters;
 		Light.LightSceneInfo->Proxy->GetLightShaderParameters(LightShaderParameters);
 
 		uint32 Transmission = Light.LightSceneInfo->Proxy->Transmission();
 		uint8 LightingChannelMask = Light.LightSceneInfo->Proxy->GetLightingChannelMask();
-		LightParameters->Flags[LightParameters->Count]  = Transmission ? PATHTRACER_FLAG_TRANSMISSION_MASK : 0;
-		LightParameters->Flags[LightParameters->Count] |= LightingChannelMask & PATHTRACER_FLAG_LIGHTING_CHANNEL_MASK;
-		LightParameters->Flags[LightParameters->Count] |= Light.LightSceneInfo->Proxy->IsInverseSquared() ? 0 : PATHTRACER_FLAG_NON_INVERSE_SQUARE_FALLOFF_MASK;
+		DestLight.Flags  = Transmission ? PATHTRACER_FLAG_TRANSMISSION_MASK : 0;
+		DestLight.Flags |= LightingChannelMask & PATHTRACER_FLAG_LIGHTING_CHANNEL_MASK;
+		DestLight.Flags |= Light.LightSceneInfo->Proxy->IsInverseSquared() ? 0 : PATHTRACER_FLAG_NON_INVERSE_SQUARE_FALLOFF_MASK;
 
-		LightParameters->FalloffExponent[LightParameters->Count] = LightShaderParameters.FalloffExponent;
-		LightParameters->Attenuation[LightParameters->Count] = LightShaderParameters.InvRadius;
-		LightParameters->IESTextureSlice[LightParameters->Count] = -1; // not used by this path at the moment
+		DestLight.FalloffExponent = LightShaderParameters.FalloffExponent;
+		DestLight.Attenuation = LightShaderParameters.InvRadius;
+		DestLight.IESTextureSlice = -1; // not used by this path at the moment
 
 		ELightComponentType LightComponentType = (ELightComponentType)Light.LightSceneInfo->Proxy->GetLightType();
 		switch (LightComponentType)
@@ -300,24 +306,24 @@ RENDERER_API void SetupLightParameters(
 		{
 			if (CVarRayTracingGlobalIlluminationDirectionalLight.GetValueOnRenderThread() == 0) continue;
 
-			LightParameters->Normal[LightParameters->Count] = LightShaderParameters.Direction;
-			LightParameters->Color[LightParameters->Count] = LightShaderParameters.Color;
-			LightParameters->Flags[LightParameters->Count] |= PATHTRACING_LIGHT_DIRECTIONAL;
+			DestLight.Normal = LightShaderParameters.Direction;
+			DestLight.Color = LightShaderParameters.Color;
+			DestLight.Flags |= PATHTRACING_LIGHT_DIRECTIONAL;
 			break;
 		}
 		case LightType_Rect:
 		{
 			if (CVarRayTracingGlobalIlluminationRectLight.GetValueOnRenderThread() == 0) continue;
 
-			LightParameters->Position[LightParameters->Count] = LightShaderParameters.Position;
-			LightParameters->Normal[LightParameters->Count] = -LightShaderParameters.Direction;
-			LightParameters->dPdu[LightParameters->Count] = FVector::CrossProduct(LightShaderParameters.Direction, LightShaderParameters.Tangent);
-			LightParameters->dPdv[LightParameters->Count] = LightShaderParameters.Tangent;
-			LightParameters->Color[LightParameters->Count] = LightShaderParameters.Color;
-			LightParameters->Dimensions[LightParameters->Count] = FVector(2.0f * LightShaderParameters.SourceRadius, 2.0f * LightShaderParameters.SourceLength, 0.0f);
-			LightParameters->RectLightBarnCosAngle[LightParameters->Count] = LightShaderParameters.RectLightBarnCosAngle;
-			LightParameters->RectLightBarnLength[LightParameters->Count] = LightShaderParameters.RectLightBarnLength;
-			LightParameters->Flags[LightParameters->Count] |= PATHTRACING_LIGHT_RECT;
+			DestLight.Position = LightShaderParameters.Position;
+			DestLight.Normal = -LightShaderParameters.Direction;
+			DestLight.dPdu = FVector::CrossProduct(LightShaderParameters.Direction, LightShaderParameters.Tangent);
+			DestLight.dPdv = LightShaderParameters.Tangent;
+			DestLight.Color = LightShaderParameters.Color;
+			DestLight.Dimensions = FVector(2.0f * LightShaderParameters.SourceRadius, 2.0f * LightShaderParameters.SourceLength, 0.0f);
+			DestLight.RectLightBarnCosAngle = LightShaderParameters.RectLightBarnCosAngle;
+			DestLight.RectLightBarnLength = LightShaderParameters.RectLightBarnLength;
+			DestLight.Flags |= PATHTRACING_LIGHT_RECT;
 			break;
 		}
 		case LightType_Point:
@@ -325,31 +331,40 @@ RENDERER_API void SetupLightParameters(
 		{
 			if (CVarRayTracingGlobalIlluminationPointLight.GetValueOnRenderThread() == 0) continue;
 
-			LightParameters->Position[LightParameters->Count] = LightShaderParameters.Position;
+			DestLight.Position = LightShaderParameters.Position;
 			// #dxr_todo: UE-72556 define these differences from Lit..
-			LightParameters->Color[LightParameters->Count] = LightShaderParameters.Color;
+			DestLight.Color = LightShaderParameters.Color;
 			float SourceRadius = 0.0; // LightShaderParameters.SourceRadius causes too much noise for little pay off at this time
-			LightParameters->Dimensions[LightParameters->Count] = FVector(0.0, 0.0, SourceRadius);
-			LightParameters->Flags[LightParameters->Count] |= PATHTRACING_LIGHT_POINT;
+			DestLight.Dimensions = FVector(0.0, 0.0, SourceRadius);
+			DestLight.Flags |= PATHTRACING_LIGHT_POINT;
 			break;
 		}
 		case LightType_Spot:
 		{
 			if (CVarRayTracingGlobalIlluminationSpotLight.GetValueOnRenderThread() == 0) continue;
 
-			LightParameters->Position[LightParameters->Count] = LightShaderParameters.Position;
-			LightParameters->Normal[LightParameters->Count] = -LightShaderParameters.Direction;
+			DestLight.Position = LightShaderParameters.Position;
+			DestLight.Normal = -LightShaderParameters.Direction;
 			// #dxr_todo: UE-72556 define these differences from Lit..
-			LightParameters->Color[LightParameters->Count] = LightShaderParameters.Color;
+			DestLight.Color = LightShaderParameters.Color;
 			float SourceRadius = 0.0; // LightShaderParameters.SourceRadius causes too much noise for little pay off at this time
-			LightParameters->Dimensions[LightParameters->Count] = FVector(LightShaderParameters.SpotAngles, SourceRadius);
-			LightParameters->Flags[LightParameters->Count] |= PATHTRACING_LIGHT_SPOT;
+			DestLight.Dimensions = FVector(LightShaderParameters.SpotAngles, SourceRadius);
+			DestLight.Flags |= PATHTRACING_LIGHT_SPOT;
 			break;
 		}
 		};
 
-		LightParameters->Color[LightParameters->Count] *= Light.LightSceneInfo->Proxy->GetIndirectLightingScale();
-		LightParameters->Count++;
+		DestLight.Color *= Light.LightSceneInfo->Proxy->GetIndirectLightingScale();
+
+		// we definitely added the light if we reach this point
+		LightCount++;
+	}
+
+	{
+		// Upload the buffer of lights to the GPU (send at least one)
+		size_t DataSize = sizeof(FPathTracingLight) * FMath::Max(LightCount, 1u);
+		*OutLightBuffer = GraphBuilder.CreateSRV(FRDGBufferSRVDesc(CreateStructuredBuffer(GraphBuilder, TEXT("RTGILightsBuffer"), sizeof(FPathTracingLight), FMath::Max(LightCount, 1u), Lights, DataSize)));
+		*OutLightCount = LightCount;
 	}
 }
 
@@ -460,7 +475,8 @@ class FGlobalIlluminationRGS : public FGlobalShader
 		SHADER_PARAMETER_RDG_TEXTURE_UAV(RWTexture2D<float4>, RWGlobalIlluminationUAV)
 		SHADER_PARAMETER_RDG_TEXTURE_UAV(RWTexture2D<float2>, RWGlobalIlluminationRayDistanceUAV)
 		SHADER_PARAMETER_STRUCT_REF(FViewUniformShaderParameters, ViewUniformBuffer)
-		SHADER_PARAMETER_STRUCT_REF(FPathTracingLightData, LightParameters)
+		SHADER_PARAMETER_RDG_BUFFER_SRV(StructuredBuffer<FPathTracingLight>, SceneLights)
+		SHADER_PARAMETER(uint32, SceneLightCount)
 		SHADER_PARAMETER_STRUCT_REF(FSkyLightData, SkyLight)
 		SHADER_PARAMETER_STRUCT_INCLUDE(FSceneTextureParameters, SceneTextures)
 		SHADER_PARAMETER_RDG_TEXTURE(Texture2D, SSProfilesTexture)
@@ -519,7 +535,8 @@ class FRayTracingGlobalIlluminationCreateGatherPointsRGS : public FGlobalShader
 		SHADER_PARAMETER_STRUCT_REF(FViewUniformShaderParameters, ViewUniformBuffer)
 
 		// Light data
-		SHADER_PARAMETER_STRUCT_REF(FPathTracingLightData, LightParameters)
+		SHADER_PARAMETER_RDG_BUFFER_SRV(StructuredBuffer<FPathTracingLight>, SceneLights)
+		SHADER_PARAMETER(uint32, SceneLightCount)
 		SHADER_PARAMETER_STRUCT_REF(FSkyLightData, SkyLight)
 
 		// Shading data
@@ -585,7 +602,8 @@ class FRayTracingGlobalIlluminationCreateGatherPointsTraceRGS : public FGlobalSh
 		SHADER_PARAMETER_STRUCT_REF(FViewUniformShaderParameters, ViewUniformBuffer)
 
 		// Light data
-		SHADER_PARAMETER_STRUCT_REF(FPathTracingLightData, LightParameters)
+		SHADER_PARAMETER_RDG_BUFFER_SRV(StructuredBuffer<FPathTracingLight>, SceneLights)
+		SHADER_PARAMETER(uint32, SceneLightCount)
 		SHADER_PARAMETER_STRUCT_REF(FSkyLightData, SkyLight)
 
 		// Shading data
@@ -835,7 +853,8 @@ void CopyGatherPassParameters(
 	NewParameters->TLAS = PassParameters.TLAS;
 	NewParameters->ViewUniformBuffer = PassParameters.ViewUniformBuffer;
 
-	NewParameters->LightParameters = PassParameters.LightParameters;
+	NewParameters->SceneLights = PassParameters.SceneLights;
+	NewParameters->SceneLightCount = PassParameters.SceneLightCount;
 	NewParameters->SkyLight = PassParameters.SkyLight;
 
 	NewParameters->SceneTextures = PassParameters.SceneTextures;
@@ -875,7 +894,8 @@ void CopyGatherPassParameters(
 	NewParameters->TLAS = PassParameters.TLAS;
 	NewParameters->ViewUniformBuffer = PassParameters.ViewUniformBuffer;
 
-	NewParameters->LightParameters = PassParameters.LightParameters;
+	NewParameters->SceneLightCount = PassParameters.SceneLightCount;
+	NewParameters->SceneLights = PassParameters.SceneLights;
 	NewParameters->SkyLight = PassParameters.SkyLight;
 
 	NewParameters->SceneTextures = PassParameters.SceneTextures;
@@ -912,8 +932,6 @@ void FDeferredShadingSceneRenderer::RayTracingGlobalIlluminationCreateGatherPoin
 	int32 GatherFilterWidth = FMath::Max(CVarRayTracingGlobalIlluminationFinalGatherFilterWidth.GetValueOnRenderThread(), 0);
 	GatherFilterWidth = GatherFilterWidth * 2 + 1;
 
-	FPathTracingLightData LightParameters;
-	SetupLightParameters(*Scene, View, &LightParameters);
 
 	float MaxShadowDistance = 1.0e27;
 	if (GRayTracingGlobalIlluminationMaxShadowDistance > 0.0)
@@ -952,7 +970,7 @@ void FDeferredShadingSceneRenderer::RayTracingGlobalIlluminationCreateGatherPoin
 	PassParameters->ViewUniformBuffer = View.ViewUniformBuffer;
 
 	// Light data
-	PassParameters->LightParameters = CreateUniformBufferImmediate(LightParameters, EUniformBufferUsage::UniformBuffer_SingleFrame);
+	SetupLightParameters(*Scene, View, GraphBuilder, &PassParameters->SceneLights, &PassParameters->SceneLightCount);
 	PassParameters->SceneTextures = SceneTextures;
 	PassParameters->SkyLight = CreateUniformBufferImmediate(SkyLightParameters, EUniformBufferUsage::UniformBuffer_SingleFrame);
 
@@ -1240,9 +1258,6 @@ void FDeferredShadingSceneRenderer::RenderRayTracingGlobalIlluminationBruteForce
 
 	int32 RayTracingGISamplesPerPixel = GetRayTracingGlobalIlluminationSamplesPerPixel(View);
 
-	FPathTracingLightData LightParameters;
-	SetupLightParameters(*Scene, View, &LightParameters);
-
 	float MaxShadowDistance = 1.0e27;
 	if (GRayTracingGlobalIlluminationMaxShadowDistance > 0.0)
 	{
@@ -1278,7 +1293,7 @@ void FDeferredShadingSceneRenderer::RenderRayTracingGlobalIlluminationBruteForce
 	PassParameters->NextEventEstimationSamples = GRayTracingGlobalIlluminationNextEventEstimationSamples;
 	PassParameters->TLAS = View.RayTracingScene.RayTracingSceneRHI->GetShaderResourceView();
 	PassParameters->ViewUniformBuffer = View.ViewUniformBuffer;
-	PassParameters->LightParameters = CreateUniformBufferImmediate(LightParameters, EUniformBufferUsage::UniformBuffer_SingleDraw);
+	SetupLightParameters(*Scene, View, GraphBuilder, &PassParameters->SceneLights, &PassParameters->SceneLightCount);
 	PassParameters->SceneTextures = SceneTextures;
 	PassParameters->SkyLight = CreateUniformBufferImmediate(SkyLightParameters, EUniformBufferUsage::UniformBuffer_SingleDraw);
 
