Fix FLightmapGBufferPS compilation with Nanite vertex factory

diff --git a/Engine/Plugins/Experimental/GPULightmass/Shaders/Private/LightmapGBuffer.usf b/Engine/Plugins/Experimental/GPULightmass/Shaders/Private/LightmapGBuffer.usf
--- a/Engine/Plugins/Experimental/GPULightmass/Shaders/Private/LightmapGBuffer.usf
+++ b/Engine/Plugins/Experimental/GPULightmass/Shaders/Private/LightmapGBuffer.usf
@@ -24,6 +24,7 @@ struct FLightmapGBufferVSToPS
 	// float4 SavedWorldPositionWithShaderOffsets : POSITION2;
 }; 
 
+#if VERTEXSHADER
 void LightmapGBufferVS(
 	FVertexFactoryInput Input,
 	out FLightmapGBufferVSToPS Output
@@ -62,6 +63,7 @@ void LightmapGBufferVS(
 	// Support WPO in the far future
 	// Output.SavedWorldPositionWithShaderOffsets = Output.SavedWorldPosition + float4(GetMaterialWorldPositionOffset(VertexParameters), 0.0f);
 }
+#endif
 
 #define WorldPositionScalar 0.00001f
 
