Add additional missing files from previous submit.

diff --git a/Engine/Plugins/Experimental/GPULightmass/Source/GPULightmass/Private/Scene/Scene.cpp b/Engine/Plugins/Experimental/GPULightmass/Source/GPULightmass/Private/Scene/Scene.cpp
--- a/Engine/Plugins/Experimental/GPULightmass/Source/GPULightmass/Private/Scene/Scene.cpp
+++ b/Engine/Plugins/Experimental/GPULightmass/Source/GPULightmass/Private/Scene/Scene.cpp
@@ -804,9 +804,12 @@ void FScene::AddGeometryInstanceFromComponent(UStaticMeshComponent* InComponent)
 			0,
 			INDEX_NONE,
 			INDEX_NONE,
-			false,
-			nullptr
-			);
+			/* bOutputVelocity = */ false,
+			nullptr,
+			/* bCastContactShadow = */ true,
+			INDEX_NONE,
+			0
+		);
 
 		for (int32 LODIndex = 0; LODIndex < InstanceLightmapRenderStateInitializers.Num(); LODIndex++)
 		{
@@ -1078,8 +1081,12 @@ void FScene::AddGeometryInstanceFromComponent(UInstancedStaticMeshComponent* InC
 			0,
 			INDEX_NONE,
 			INDEX_NONE,
-			false,
-			nullptr), UniformBuffer_MultiFrame);
+			/* bOutputVelocity = */ false,
+			nullptr,
+			/* bCastContactShadow = */ true,
+			INDEX_NONE,
+			0
+		), UniformBuffer_MultiFrame);
 
 		for (int32 LODIndex = 0; LODIndex < InstanceLightmapRenderStateInitializers.Num(); LODIndex++)
 		{
@@ -1368,8 +1375,12 @@ void FScene::AddGeometryInstanceFromComponent(ULandscapeComponent* InComponent)
 			0,
 			INDEX_NONE,
 			INDEX_NONE,
-			false,
-			nullptr), UniformBuffer_MultiFrame);
+			/* bOutputVelocity = */ false,
+			nullptr,
+			/* bCastContactShadow = */ true,
+			INDEX_NONE,
+			0
+		), UniformBuffer_MultiFrame);
 
 		int32 MaxLOD = 0;
 		InstanceRenderStateRef->LandscapeFixedGridUniformShaderParameters.AddDefaulted(MaxLOD + 1);
diff --git a/Engine/Shaders/Private/InstanceCulling/BuildInstanceDrawCommands.usf b/Engine/Shaders/Private/InstanceCulling/BuildInstanceDrawCommands.usf
--- a/Engine/Shaders/Private/InstanceCulling/BuildInstanceDrawCommands.usf
+++ b/Engine/Shaders/Private/InstanceCulling/BuildInstanceDrawCommands.usf
@@ -1,9 +1,5 @@
 // Copyright Epic Games, Inc. All Rights Reserved.
 
-/*=============================================================================
-	BuildInstanceDrawCommands.usf: 
-=============================================================================*/
-
 #include "../Common.ush"
 #include "../SceneData.ush"
 
diff --git a/Engine/Shaders/Private/InstanceCulling/CullInstances.usf b/Engine/Shaders/Private/InstanceCulling/CullInstances.usf
--- a/Engine/Shaders/Private/InstanceCulling/CullInstances.usf
+++ b/Engine/Shaders/Private/InstanceCulling/CullInstances.usf
@@ -1,8 +1,4 @@
-// Copyright 1998-2019 Epic Games, Inc. All Rights Reserved.
-
-/*=============================================================================
-	CullInstances.usf: 
-=============================================================================*/
+// Copyright Epic Games, Inc. All Rights Reserved.
 
 #include "../Common.ush"
 #include "../SceneData.ush"
diff --git a/Engine/Source/Runtime/Engine/Private/DynamicMeshBuilder.cpp b/Engine/Source/Runtime/Engine/Private/DynamicMeshBuilder.cpp
--- a/Engine/Source/Runtime/Engine/Private/DynamicMeshBuilder.cpp
+++ b/Engine/Source/Runtime/Engine/Private/DynamicMeshBuilder.cpp
@@ -870,7 +870,11 @@ void FDynamicMeshBuilder::GetMesh(
 				INDEX_NONE,
 				INDEX_NONE,
 				INDEX_NONE,
-				false
+				/* bOutputVelocity = */ false,
+				nullptr,
+				/* bCastContactShadow = */ true,
+				INDEX_NONE,
+				0
 			);
 
 
@@ -963,10 +967,13 @@ void FDynamicMeshBuilder::GetMeshElement(const FMatrix& LocalToWorld, const FMat
 				INDEX_NONE,
 				INDEX_NONE,
 				INDEX_NONE,
-				false
+				/* bOutputVelocity = */ false,
+				nullptr,
+				/* bCastContactShadow = */ true,
+				INDEX_NONE,
+				0
 			);
 
-
 			if (IsInGameThread())
 			{
 				BeginSetUniformBufferContents(*OneFrameResource.PrimitiveUniformBuffer, PrimitiveParams);
@@ -1043,8 +1050,12 @@ void FDynamicMeshBuilder::Draw(FPrimitiveDrawInterface* PDI,const FMatrix& Local
 			INDEX_NONE,
 			INDEX_NONE,
 			INDEX_NONE,
-			false
-			);
+			/* bOutputVelocity = */ false,
+			nullptr,
+			/* bCastContactShadow = */ true,
+			INDEX_NONE,
+			0
+		);
 
 
 		if (IsInGameThread())
@@ -1079,7 +1090,7 @@ void FDynamicMeshBuilder::Draw(FPrimitiveDrawInterface* PDI,const FMatrix& Local
 
 		// Clear the resource pointers so they cannot be overwritten accidentally.
 		// These resources will be released by the PDI.
-		VertexBuffer = NULL;
-		IndexBuffer = NULL;
+		VertexBuffer = nullptr;
+		IndexBuffer = nullptr;
 	}
 }
diff --git a/Engine/Source/Runtime/Engine/Private/Particles/ParticleSystemRender.cpp b/Engine/Source/Runtime/Engine/Private/Particles/ParticleSystemRender.cpp
--- a/Engine/Source/Runtime/Engine/Private/Particles/ParticleSystemRender.cpp
+++ b/Engine/Source/Runtime/Engine/Private/Particles/ParticleSystemRender.cpp
@@ -7212,8 +7212,13 @@ void FParticleSystemSceneProxy::UpdateWorldSpacePrimitiveUniformBuffer() const
 			INDEX_NONE,
 			INDEX_NONE,
 			INDEX_NONE,
-			AlwaysHasVelocity()
-			);
+			AlwaysHasVelocity(),
+			nullptr,
+			/* bCastContactShadow = */ true,
+			INDEX_NONE,
+			0
+		);
+
 		WorldSpacePrimitiveUniformBuffer.SetContents(PrimitiveUniformShaderParameters);
 		WorldSpacePrimitiveUniformBuffer.InitResource();
 	}
