Fix for build w/ experimental plugin

diff --git a/Engine/Plugins/Experimental/GPULightmass/Source/GPULightmass/Private/Scene/Scene.cpp b/Engine/Plugins/Experimental/GPULightmass/Source/GPULightmass/Private/Scene/Scene.cpp
--- a/Engine/Plugins/Experimental/GPULightmass/Source/GPULightmass/Private/Scene/Scene.cpp
+++ b/Engine/Plugins/Experimental/GPULightmass/Source/GPULightmass/Private/Scene/Scene.cpp
@@ -808,7 +808,8 @@ void FScene::AddGeometryInstanceFromComponent(UStaticMeshComponent* InComponent)
 			nullptr,
 			/* bCastContactShadow = */ true,
 			INDEX_NONE,
-			0
+			0,
+			/* bCastShadow = */ true
 		);
 
 		for (int32 LODIndex = 0; LODIndex < InstanceLightmapRenderStateInitializers.Num(); LODIndex++)
@@ -1085,7 +1086,8 @@ void FScene::AddGeometryInstanceFromComponent(UInstancedStaticMeshComponent* InC
 			nullptr,
 			/* bCastContactShadow = */ true,
 			INDEX_NONE,
-			0
+			0,
+			/* bCastShadow = */ true
 		), UniformBuffer_MultiFrame);
 
 		for (int32 LODIndex = 0; LODIndex < InstanceLightmapRenderStateInitializers.Num(); LODIndex++)
@@ -1379,7 +1381,8 @@ void FScene::AddGeometryInstanceFromComponent(ULandscapeComponent* InComponent)
 			nullptr,
 			/* bCastContactShadow = */ true,
 			INDEX_NONE,
-			0
+			0,
+			/* bCastShadow = */ true
 		), UniformBuffer_MultiFrame);
 
 		int32 MaxLOD = 0;
