Refactor path tracer light struct to provide soft radius and source length information

No functionality change in this review, just getting the code ready for implementing these features

diff --git a/Engine/Plugins/Experimental/GPULightmass/Source/GPULightmass/Private/LightmapRenderer.cpp b/Engine/Plugins/Experimental/GPULightmass/Source/GPULightmass/Private/LightmapRenderer.cpp
--- a/Engine/Plugins/Experimental/GPULightmass/Source/GPULightmass/Private/LightmapRenderer.cpp
+++ b/Engine/Plugins/Experimental/GPULightmass/Source/GPULightmass/Private/LightmapRenderer.cpp
@@ -191,7 +191,10 @@ void SetupPathTracingLightParameters(const GPULightmass::FLightSceneRenderState&
 
 			DestLight.Normal = -Light.Direction;
 			DestLight.Color = FVector(Light.Color);
-			DestLight.Dimensions = FVector(0.0f, 0.0f, FMath::Sin(0.5f * FMath::DegreesToRadians(Light.LightSourceAngle)));
+			DestLight.Dimensions = FVector(
+				FMath::Sin(0.5f * FMath::DegreesToRadians(Light.LightSourceAngle)),
+				FMath::Sin(0.5f * FMath::DegreesToRadians(Light.LightSourceSoftAngle)),
+				0.0f);
 			DestLight.Attenuation = 1.0;
 			DestLight.IESTextureSlice = -1;
 
@@ -210,7 +213,11 @@ void SetupPathTracingLightParameters(const GPULightmass::FLightSceneRenderState&
 
 			DestLight.Position = Light.Position;
 			DestLight.Color = FVector(Light.Color);
-			DestLight.Dimensions = FVector(0.0, 0.0, Light.SourceRadius);
+			DestLight.Normal = Light.Direction;
+			DestLight.dPdu = FVector::CrossProduct(Light.Tangent, Light.Direction);
+			DestLight.dPdv = Light.Tangent;
+
+			DestLight.Dimensions = FVector(Light.SourceRadius, Light.SourceSoftRadius, Light.SourceLength);
 			DestLight.Attenuation = 1.0f / Light.AttenuationRadius;
 			DestLight.FalloffExponent = Light.FalloffExponent;
 			DestLight.IESTextureSlice = -1;
@@ -231,8 +238,11 @@ void SetupPathTracingLightParameters(const GPULightmass::FLightSceneRenderState&
 
 			DestLight.Position = Light.Position;
 			DestLight.Normal = Light.Direction;
+			DestLight.dPdu = FVector::CrossProduct(Light.Tangent, Light.Direction);
+			DestLight.dPdv = Light.Tangent;
 			DestLight.Color = FVector(Light.Color);
-			DestLight.Dimensions = FVector(Light.SpotAngles, Light.SourceRadius);
+			DestLight.Dimensions = FVector(Light.SourceRadius, Light.SourceSoftRadius, Light.SourceLength);
+			DestLight.Shaping = Light.SpotAngles;
 			DestLight.Attenuation = 1.0f / Light.AttenuationRadius;
 			DestLight.FalloffExponent = Light.FalloffExponent;
 			DestLight.IESTextureSlice = -1;
@@ -262,8 +272,7 @@ void SetupPathTracingLightParameters(const GPULightmass::FLightSceneRenderState&
 
 			DestLight.Dimensions = FVector(Light.SourceWidth, Light.SourceHeight, 0.0f);
 			DestLight.Attenuation = 1.0f / Light.AttenuationRadius;
-			DestLight.RectLightBarnCosAngle = FMath::Cos(FMath::DegreesToRadians(Light.BarnDoorAngle));
-			DestLight.RectLightBarnLength = Light.BarnDoorLength;
+			DestLight.Shaping = FVector2D(FMath::Cos(FMath::DegreesToRadians(Light.BarnDoorAngle)), Light.BarnDoorLength);
 
 			DestLight.IESTextureSlice = -1;
 
diff --git a/Engine/Plugins/Experimental/GPULightmass/Source/GPULightmass/Private/Scene/Lights.cpp b/Engine/Plugins/Experimental/GPULightmass/Source/GPULightmass/Private/Scene/Lights.cpp
--- a/Engine/Plugins/Experimental/GPULightmass/Source/GPULightmass/Private/Scene/Lights.cpp
+++ b/Engine/Plugins/Experimental/GPULightmass/Source/GPULightmass/Private/Scene/Lights.cpp
@@ -41,6 +41,7 @@ FDirectionalLightRenderState::FDirectionalLightRenderState(UDirectionalLightComp
 	Color = DirectionalLightComponent->GetColoredLightBrightness();
 	Direction = DirectionalLightComponent->GetDirection();
 	LightSourceAngle = DirectionalLightComponent->LightSourceAngle;
+	LightSourceSoftAngle = DirectionalLightComponent->LightSourceSoftAngle;
 	ShadowMapChannel = DirectionalLightComponent->PreviewShadowMapChannel;
 }
 
@@ -69,8 +70,15 @@ FPointLightRenderState::FPointLightRenderState(UPointLightComponent* PointLightC
 	bStationary = bCastStationaryShadows;
 	Color = PointLightComponent->GetColoredLightBrightness();
 	Position = PointLightComponent->GetLightPosition();
+	Direction = PointLightComponent->GetDirection();
+	{
+		FMatrix LightToWorld = PointLightComponent->GetComponentTransform().ToMatrixNoScale();
+		Tangent = FVector(LightToWorld.M[2][0], LightToWorld.M[2][1], LightToWorld.M[2][2]);
+	}
 	AttenuationRadius = PointLightComponent->AttenuationRadius;
 	SourceRadius = PointLightComponent->SourceRadius;
+	SourceSoftRadius = PointLightComponent->SoftSourceRadius;
+	SourceLength = PointLightComponent->SourceLength;
 	ShadowMapChannel = PointLightComponent->PreviewShadowMapChannel;
 	FalloffExponent = PointLightComponent->LightFalloffExponent;
 	IsInverseSquared = PointLightComponent->bUseInverseSquaredFalloff;
@@ -143,6 +151,8 @@ FSpotLightRenderState::FSpotLightRenderState(USpotLightComponent* SpotLightCompo
 	}
 	AttenuationRadius = SpotLightComponent->AttenuationRadius;
 	SourceRadius = SpotLightComponent->SourceRadius;
+	SourceSoftRadius = SpotLightComponent->SoftSourceRadius;
+	SourceLength = SpotLightComponent->SourceLength;
 	ShadowMapChannel = SpotLightComponent->PreviewShadowMapChannel;
 	FalloffExponent = SpotLightComponent->LightFalloffExponent;
 	IsInverseSquared = SpotLightComponent->bUseInverseSquaredFalloff;
diff --git a/Engine/Plugins/Experimental/GPULightmass/Source/GPULightmass/Private/Scene/Lights.h b/Engine/Plugins/Experimental/GPULightmass/Source/GPULightmass/Private/Scene/Lights.h
--- a/Engine/Plugins/Experimental/GPULightmass/Source/GPULightmass/Private/Scene/Lights.h
+++ b/Engine/Plugins/Experimental/GPULightmass/Source/GPULightmass/Private/Scene/Lights.h
@@ -203,6 +203,7 @@ struct FDirectionalLightRenderState : public FLocalLightRenderState
 	FVector Direction;
 	FLinearColor Color;
 	float LightSourceAngle;
+	float LightSourceSoftAngle;
 
 	virtual FLightShaderParameters GetLightShaderParameters() const override;
 };
@@ -214,9 +215,13 @@ struct FPointLightRenderState : public FLocalLightRenderState
 	FPointLightRenderState(UPointLightComponent* PointLightComponent);
 
 	FVector Position;
+	FVector Direction;
+	FVector Tangent;
 	FLinearColor Color;
 	float AttenuationRadius;
 	float SourceRadius;
+	float SourceSoftRadius;
+	float SourceLength;
 	float FalloffExponent;
 	bool IsInverseSquared;
 
@@ -236,6 +241,8 @@ struct FSpotLightRenderState : public FLocalLightRenderState
 	FLinearColor Color;
 	float AttenuationRadius;
 	float SourceRadius;
+	float SourceSoftRadius;
+	float SourceLength;
 	float FalloffExponent;
 	bool IsInverseSquared;
 
diff --git a/Engine/Shaders/Private/PathTracing/Light/PathTracingLightCommon.ush b/Engine/Shaders/Private/PathTracing/Light/PathTracingLightCommon.ush
--- a/Engine/Shaders/Private/PathTracing/Light/PathTracingLightCommon.ush
+++ b/Engine/Shaders/Private/PathTracing/Light/PathTracingLightCommon.ush
@@ -93,14 +93,9 @@ float2 GetRectSize(int LightId)
 	return SceneLights[LightId].Dimensions.xy;
 }
 
-float GetDepth(int LightId)
-{
-	return SceneLights[LightId].Dimensions.z;
-}
-
 float2 GetCosConeAngles(int LightId)
 {
-	return SceneLights[LightId].Dimensions.xy;
+	return SceneLights[LightId].Shaping;
 }
 
 float3 GetColor(int LightId)
@@ -109,6 +104,16 @@ float3 GetColor(int LightId)
 }
 
 float GetRadius(int LightId)
+{
+	return SceneLights[LightId].Dimensions.x;
+}
+
+float GetSoftRadius(int LightId)
+{
+	return SceneLights[LightId].Dimensions.y;
+}
+
+float GetSourceLength(int LightId)
 {
 	return SceneLights[LightId].Dimensions.z;
 }
@@ -120,12 +125,12 @@ float GetAttenuation(int LightId)
 
 float GetRectLightBarnCosAngle(int LightId)
 {
-	return SceneLights[LightId].RectLightBarnCosAngle;
+	return SceneLights[LightId].Shaping.x;
 }
 
 float GetRectLightBarnLength(int LightId)
 {
-	return SceneLights[LightId].RectLightBarnLength;
+	return SceneLights[LightId].Shaping.y;
 }
 
 #ifndef USE_ATTENUATION_TERM
diff --git a/Engine/Shaders/Shared/RayTracingTypes.h b/Engine/Shaders/Shared/RayTracingTypes.h
--- a/Engine/Shaders/Shared/RayTracingTypes.h
+++ b/Engine/Shaders/Shared/RayTracingTypes.h
@@ -18,10 +18,9 @@ struct FPathTracingLight {
 	FVector dPdv;
 	FVector Color;
 	FVector Dimensions; // Radius,SoftRadius,Length or RectWidth,RectHeight or Sin(Angle/2),Sin(SoftAngle/2) depending on light type
+	FVector2D Shaping;  // Barndoor controls for RectLights, Cone angles for spots lights
 	float   Attenuation;
 	float   FalloffExponent; // for non-inverse square decay lights only
-	float   RectLightBarnCosAngle;
-	float   RectLightBarnLength;
 	int32   IESTextureSlice;
 	uint32  Flags; // see defines above
 
@@ -41,10 +40,9 @@ struct FPathTracingLight {
 	float3  dPdv;
 	float3  Color;
 	float3  Dimensions;
+	float2  Shaping;
 	float   Attenuation;
 	float   FalloffExponent;
-	float   RectLightBarnCosAngle;
-	float   RectLightBarnLength;
 	int     IESTextureSlice;
 	uint    Flags;
 
diff --git a/Engine/Source/Runtime/Renderer/Private/PathTracing.cpp b/Engine/Source/Runtime/Renderer/Private/PathTracing.cpp
--- a/Engine/Source/Runtime/Renderer/Private/PathTracing.cpp
+++ b/Engine/Source/Runtime/Renderer/Private/PathTracing.cpp
@@ -317,15 +317,14 @@ void SetLightParameters(FRDGBuilder& GraphBuilder, FPathTracingRG::FParameters*
 			case LightType_Directional:
 			{
 				DestLight.Normal = LightParameters.Direction;
-				DestLight.Dimensions = FVector(0.0f, 0.0f, LightParameters.SourceRadius);
+				DestLight.Dimensions = FVector(LightParameters.SourceRadius, LightParameters.SoftSourceRadius, 0.0f);
 				DestLight.Flags |= PATHTRACING_LIGHT_DIRECTIONAL;
 				break;
 			}
 			case LightType_Rect:
 			{
 				DestLight.Dimensions = FVector(2.0f * LightParameters.SourceRadius, 2.0f * LightParameters.SourceLength, 0.0f);
-				DestLight.RectLightBarnCosAngle = LightParameters.RectLightBarnCosAngle;
-				DestLight.RectLightBarnLength = LightParameters.RectLightBarnLength;
+				DestLight.Shaping = FVector2D(LightParameters.RectLightBarnCosAngle, LightParameters.RectLightBarnLength);
 				DestLight.FalloffExponent = LightParameters.FalloffExponent;
 				DestLight.Flags |= Light.LightSceneInfo->Proxy->IsInverseSquared() ? 0 : PATHTRACER_FLAG_NON_INVERSE_SQUARE_FALLOFF_MASK;
 				DestLight.Flags |= PATHTRACING_LIGHT_RECT;
@@ -333,7 +332,8 @@ void SetLightParameters(FRDGBuilder& GraphBuilder, FPathTracingRG::FParameters*
 			}
 			case LightType_Spot:
 			{
-				DestLight.Dimensions = FVector(LightParameters.SpotAngles, LightParameters.SourceRadius);
+				DestLight.Dimensions = FVector(LightParameters.SourceRadius, LightParameters.SoftSourceRadius, LightParameters.SourceLength);
+				DestLight.Shaping = LightParameters.SpotAngles;
 				DestLight.FalloffExponent = LightParameters.FalloffExponent;
 				DestLight.Flags |= Light.LightSceneInfo->Proxy->IsInverseSquared() ? 0 : PATHTRACER_FLAG_NON_INVERSE_SQUARE_FALLOFF_MASK;
 				DestLight.Flags |= PATHTRACING_LIGHT_SPOT;
@@ -341,7 +341,7 @@ void SetLightParameters(FRDGBuilder& GraphBuilder, FPathTracingRG::FParameters*
 			}
 			case LightType_Point:
 			{
-				DestLight.Dimensions = FVector(0.0, 0.0, LightParameters.SourceRadius);
+				DestLight.Dimensions = FVector(LightParameters.SourceRadius, LightParameters.SoftSourceRadius, LightParameters.SourceLength);
 				DestLight.FalloffExponent = LightParameters.FalloffExponent;
 				DestLight.Flags |= Light.LightSceneInfo->Proxy->IsInverseSquared() ? 0 : PATHTRACER_FLAG_NON_INVERSE_SQUARE_FALLOFF_MASK;
 				DestLight.Flags |= PATHTRACING_LIGHT_POINT;
diff --git a/Engine/Source/Runtime/Renderer/Private/RayTracing/RayTracingGlobalIllumination.cpp b/Engine/Source/Runtime/Renderer/Private/RayTracing/RayTracingGlobalIllumination.cpp
--- a/Engine/Source/Runtime/Renderer/Private/RayTracing/RayTracingGlobalIllumination.cpp
+++ b/Engine/Source/Runtime/Renderer/Private/RayTracing/RayTracingGlobalIllumination.cpp
@@ -321,8 +321,7 @@ RENDERER_API void SetupLightParameters(
 			DestLight.dPdv = LightShaderParameters.Tangent;
 			DestLight.Color = LightShaderParameters.Color;
 			DestLight.Dimensions = FVector(2.0f * LightShaderParameters.SourceRadius, 2.0f * LightShaderParameters.SourceLength, 0.0f);
-			DestLight.RectLightBarnCosAngle = LightShaderParameters.RectLightBarnCosAngle;
-			DestLight.RectLightBarnLength = LightShaderParameters.RectLightBarnLength;
+			DestLight.Shaping = FVector2D(LightShaderParameters.RectLightBarnCosAngle, LightShaderParameters.RectLightBarnLength);
 			DestLight.Flags |= PATHTRACING_LIGHT_RECT;
 			break;
 		}
@@ -335,7 +334,7 @@ RENDERER_API void SetupLightParameters(
 			// #dxr_todo: UE-72556 define these differences from Lit..
 			DestLight.Color = LightShaderParameters.Color;
 			float SourceRadius = 0.0; // LightShaderParameters.SourceRadius causes too much noise for little pay off at this time
-			DestLight.Dimensions = FVector(0.0, 0.0, SourceRadius);
+			DestLight.Dimensions = FVector(SourceRadius, 0.0, 0.0);
 			DestLight.Flags |= PATHTRACING_LIGHT_POINT;
 			break;
 		}
@@ -348,7 +347,8 @@ RENDERER_API void SetupLightParameters(
 			// #dxr_todo: UE-72556 define these differences from Lit..
 			DestLight.Color = LightShaderParameters.Color;
 			float SourceRadius = 0.0; // LightShaderParameters.SourceRadius causes too much noise for little pay off at this time
-			DestLight.Dimensions = FVector(LightShaderParameters.SpotAngles, SourceRadius);
+			DestLight.Dimensions = FVector(SourceRadius, 0.0, 0.0);
+			DestLight.Shaping = LightShaderParameters.SpotAngles;
 			DestLight.Flags |= PATHTRACING_LIGHT_SPOT;
 			break;
 		}
