GPULM: Use the new SH denoiser in Intel OpenImageDenoise 1.3.0 to denoise lightmap directionality

diff --git a/Engine/Plugins/Experimental/GPULightmass/Source/GPULightmass/Private/LightmapDenoising.cpp b/Engine/Plugins/Experimental/GPULightmass/Source/GPULightmass/Private/LightmapDenoising.cpp
--- a/Engine/Plugins/Experimental/GPULightmass/Source/GPULightmass/Private/LightmapDenoising.cpp
+++ b/Engine/Plugins/Experimental/GPULightmass/Source/GPULightmass/Private/LightmapDenoising.cpp
@@ -3,7 +3,7 @@
 #include "LightmapDenoising.h"
 #include "GPULightmassCommon.h"
 
-FDenoiserFilterSet::FDenoiserFilterSet(FDenoiserContext& Context, FIntPoint NewSize)
+FDenoiserFilterSet::FDenoiserFilterSet(FDenoiserContext& Context, FIntPoint NewSize, bool bSHDenoiser)
 	: Context(Context)
 {
 	double Start = FPlatformTime::Seconds();
@@ -19,7 +19,8 @@ FDenoiserFilterSet::FDenoiserFilterSet(FDenoiserContext& Context, FIntPoint NewS
 	filter = Context.OIDNDevice.newFilter("RTLightmap");
 	filter.setImage("color", InputBuffer.GetData(), oidn::Format::Float3, Size.X, Size.Y);
 	filter.setImage("output", OutputBuffer.GetData(), oidn::Format::Float3, Size.X, Size.Y);
-	filter.set("hdr", true);
+	filter.set("directional", bSHDenoiser);
+	filter.set("inputScale", bSHDenoiser ? 0.5f : 1.0f);
 	filter.commit();
 #endif
 
@@ -95,11 +96,12 @@ void DenoiseLightSampleData(FIntPoint Size, TArray<FLightSampleData>& LightSampl
 		}
 	}
 
-	// Resizing the filter is a super expensive operation
-	// Round things into size bins to reduce number of resizes
-	FDenoiserFilterSet& FilterSet = DenoiserContext.GetFilterForSize(FIntPoint(FMath::DivideAndRoundUp(Size.X, 64) * 64, FMath::DivideAndRoundUp(Size.Y, 64) * 64));
-
 	{
+
+		// Resizing the filter is a super expensive operation
+		// Round things into size bins to reduce number of resizes
+		FDenoiserFilterSet& FilterSet = DenoiserContext.GetFilterForSize(FIntPoint(FMath::DivideAndRoundUp(Size.X, 64) * 64, FMath::DivideAndRoundUp(Size.Y, 64) * 64));
+
 		for (int32 Y = 0; Y < Size.Y; Y++)
 		{
 			for (int32 X = 0; X < Size.X; X++)
@@ -124,25 +126,17 @@ void DenoiseLightSampleData(FIntPoint Size, TArray<FLightSampleData>& LightSampl
 	}
 
 	{
-		FLinearColor MinValue(MAX_flt, MAX_flt, MAX_flt);
-
-		for (int32 Y = 0; Y < Size.Y; Y++)
-		{
-			for (int32 X = 0; X < Size.X; X++)
-			{
-				MinValue.R = FMath::Min(MinValue.R, LightSampleData[Y * Size.X + X].Coefficients[1][0]);
-				MinValue.G = FMath::Min(MinValue.G, LightSampleData[Y * Size.X + X].Coefficients[1][1]);
-				MinValue.B = FMath::Min(MinValue.B, LightSampleData[Y * Size.X + X].Coefficients[1][2]);
-			}
-		}
+		// Resizing the filter is a super expensive operation
+		// Round things into size bins to reduce number of resizes
+		FDenoiserFilterSet& FilterSet = DenoiserContext.GetFilterForSize(FIntPoint(FMath::DivideAndRoundUp(Size.X, 64) * 64, FMath::DivideAndRoundUp(Size.Y, 64) * 64), true);
 
 		for (int32 Y = 0; Y < Size.Y; Y++)
 		{
 			for (int32 X = 0; X < Size.X; X++)
 			{
-				FilterSet.InputBuffer[Y * FilterSet.Size.X + X][0] = LightSampleData[Y * Size.X + X].Coefficients[1][0] - MinValue.R;
-				FilterSet.InputBuffer[Y * FilterSet.Size.X + X][1] = LightSampleData[Y * Size.X + X].Coefficients[1][1] - MinValue.G;
-				FilterSet.InputBuffer[Y * FilterSet.Size.X + X][2] = LightSampleData[Y * Size.X + X].Coefficients[1][2] - MinValue.B;
+				FilterSet.InputBuffer[Y * FilterSet.Size.X + X][0] = LightSampleData[Y * Size.X + X].Coefficients[1][0];
+				FilterSet.InputBuffer[Y * FilterSet.Size.X + X][1] = LightSampleData[Y * Size.X + X].Coefficients[1][1];
+				FilterSet.InputBuffer[Y * FilterSet.Size.X + X][2] = LightSampleData[Y * Size.X + X].Coefficients[1][2];
 			}
 		}
 
@@ -152,9 +146,9 @@ void DenoiseLightSampleData(FIntPoint Size, TArray<FLightSampleData>& LightSampl
 		{
 			for (int32 X = 0; X < Size.X; X++)
 			{
-				LightSampleData[Y * Size.X + X].Coefficients[1][0] = FilterSet.OutputBuffer[Y * FilterSet.Size.X + X][0] + MinValue.R;
-				LightSampleData[Y * Size.X + X].Coefficients[1][1] = FilterSet.OutputBuffer[Y * FilterSet.Size.X + X][1] + MinValue.G;
-				LightSampleData[Y * Size.X + X].Coefficients[1][2] = FilterSet.OutputBuffer[Y * FilterSet.Size.X + X][2] + MinValue.B;
+				LightSampleData[Y * Size.X + X].Coefficients[1][0] = FilterSet.OutputBuffer[Y * FilterSet.Size.X + X][0];
+				LightSampleData[Y * Size.X + X].Coefficients[1][1] = FilterSet.OutputBuffer[Y * FilterSet.Size.X + X][1];
+				LightSampleData[Y * Size.X + X].Coefficients[1][2] = FilterSet.OutputBuffer[Y * FilterSet.Size.X + X][2];
 			}
 		}
 	}
@@ -256,11 +250,11 @@ void DenoiseRawData(
 		}
 	}
 
-	// Resizing the filter is a super expensive operation
-	// Round things into size bins to reduce number of resizes
-	FDenoiserFilterSet& FilterSet = DenoiserContext.GetFilterForSize(FIntPoint(FMath::DivideAndRoundUp(Size.X, 64) * 64, FMath::DivideAndRoundUp(Size.Y, 64) * 64));
-
 	{
+		// Resizing the filter is a super expensive operation
+		// Round things into size bins to reduce number of resizes
+		FDenoiserFilterSet& FilterSet = DenoiserContext.GetFilterForSize(FIntPoint(FMath::DivideAndRoundUp(Size.X, 64) * 64, FMath::DivideAndRoundUp(Size.Y, 64) * 64));
+
 		for (int32 Y = 0; Y < Size.Y; Y++)
 		{
 			for (int32 X = 0; X < Size.X; X++)
@@ -285,25 +279,19 @@ void DenoiseRawData(
 	}
 
 	{
-		FLinearColor MinValue(MAX_flt, MAX_flt, MAX_flt);
+		// Resizing the filter is a super expensive operation
+		// Round things into size bins to reduce number of resizes
+		FDenoiserFilterSet& FilterSet = DenoiserContext.GetFilterForSize(FIntPoint(FMath::DivideAndRoundUp(Size.X, 64) * 64, FMath::DivideAndRoundUp(Size.Y, 64) * 64), true);
 
-		for (int32 Y = 0; Y < Size.Y; Y++)
-		{
-			for (int32 X = 0; X < Size.X; X++)
-			{
-				MinValue.R = FMath::Min(MinValue.R, LuminanceSH[Y * Size.X + X].R);
-				MinValue.G = FMath::Min(MinValue.G, LuminanceSH[Y * Size.X + X].G);
-				MinValue.B = FMath::Min(MinValue.B, LuminanceSH[Y * Size.X + X].B);
-			}
-		}
+		FLinearColor MinValue(MAX_flt, MAX_flt, MAX_flt);
 
 		for (int32 Y = 0; Y < Size.Y; Y++)
 		{
 			for (int32 X = 0; X < Size.X; X++)
 			{
-				FilterSet.InputBuffer[Y * FilterSet.Size.X + X][0] = LuminanceSH[Y * Size.X + X].R - MinValue.R;
-				FilterSet.InputBuffer[Y * FilterSet.Size.X + X][1] = LuminanceSH[Y * Size.X + X].G - MinValue.G;
-				FilterSet.InputBuffer[Y * FilterSet.Size.X + X][2] = LuminanceSH[Y * Size.X + X].B - MinValue.B;
+				FilterSet.InputBuffer[Y * FilterSet.Size.X + X][0] = LuminanceSH[Y * Size.X + X].R;
+				FilterSet.InputBuffer[Y * FilterSet.Size.X + X][1] = LuminanceSH[Y * Size.X + X].G;
+				FilterSet.InputBuffer[Y * FilterSet.Size.X + X][2] = LuminanceSH[Y * Size.X + X].B;
 			}
 		}
 
@@ -313,9 +301,9 @@ void DenoiseRawData(
 		{
 			for (int32 X = 0; X < Size.X; X++)
 			{
-				LuminanceSH[Y * Size.X + X].R = FilterSet.OutputBuffer[Y * FilterSet.Size.X + X][0] + MinValue.R;
-				LuminanceSH[Y * Size.X + X].G = FilterSet.OutputBuffer[Y * FilterSet.Size.X + X][1] + MinValue.G;
-				LuminanceSH[Y * Size.X + X].B = FilterSet.OutputBuffer[Y * FilterSet.Size.X + X][2] + MinValue.B;
+				LuminanceSH[Y * Size.X + X].R = FilterSet.OutputBuffer[Y * FilterSet.Size.X + X][0];
+				LuminanceSH[Y * Size.X + X].G = FilterSet.OutputBuffer[Y * FilterSet.Size.X + X][1];
+				LuminanceSH[Y * Size.X + X].B = FilterSet.OutputBuffer[Y * FilterSet.Size.X + X][2];
 			}
 		}
 	}
diff --git a/Engine/Plugins/Experimental/GPULightmass/Source/GPULightmass/Private/LightmapDenoising.h b/Engine/Plugins/Experimental/GPULightmass/Source/GPULightmass/Private/LightmapDenoising.h
--- a/Engine/Plugins/Experimental/GPULightmass/Source/GPULightmass/Private/LightmapDenoising.h
+++ b/Engine/Plugins/Experimental/GPULightmass/Source/GPULightmass/Private/LightmapDenoising.h
@@ -27,7 +27,7 @@ struct FDenoiserFilterSet
 	TArray<FVector> InputBuffer;
 	TArray<FVector> OutputBuffer;
 
-	FDenoiserFilterSet(FDenoiserContext& Context, FIntPoint NewSize);
+	FDenoiserFilterSet(FDenoiserContext& Context, FIntPoint NewSize, bool bSHDenoiser = false);
 
 	void Execute();
 
@@ -46,6 +46,7 @@ struct FDenoiserContext
 #endif
 
 	TMap<FIntPoint, FDenoiserFilterSet> Filters;
+	TMap<FIntPoint, FDenoiserFilterSet> SHFilters;
 
 	FDenoiserContext()
 	{
@@ -60,16 +61,30 @@ struct FDenoiserContext
 		UE_LOG(LogTemp, Log, TEXT("Denoising: %.2lfs initializing filters (%d), %.2lfs executing filters (%d)"), FilterInitTime, NumFilterInit, FilterExecutionTime, NumFilterExecution);
 	}
 
-	FDenoiserFilterSet& GetFilterForSize(FIntPoint Size)
+	FDenoiserFilterSet& GetFilterForSize(FIntPoint Size, bool bSHDenoiser = false)
 	{
-		if (!Filters.Contains(Size))
+		if (!bSHDenoiser)
 		{
-			Filters.Add(Size, FDenoiserFilterSet(*this, Size));
+			if (!Filters.Contains(Size))
+			{
+				Filters.Add(Size, FDenoiserFilterSet(*this, Size));
+			}
+
+			Filters[Size].Clear();
+
+			return Filters[Size];
 		}
+		else
+		{
+			if (!SHFilters.Contains(Size))
+			{
+				SHFilters.Add(Size, FDenoiserFilterSet(*this, Size, true));
+			}
 
-		Filters[Size].Clear();
+			SHFilters[Size].Clear();
 
-		return Filters[Size];
+			return SHFilters[Size];
+		}
 	}
 };
 
