GPULM: Fix ray hit distance with POSITIVE_INFINITY

diff --git a/Engine/Plugins/Experimental/GPULightmass/Shaders/Private/LightmapPathTracing.usf b/Engine/Plugins/Experimental/GPULightmass/Shaders/Private/LightmapPathTracing.usf
--- a/Engine/Plugins/Experimental/GPULightmass/Shaders/Private/LightmapPathTracing.usf
+++ b/Engine/Plugins/Experimental/GPULightmass/Shaders/Private/LightmapPathTracing.usf
@@ -230,7 +230,7 @@ void PathTracingKernel(
 		else
 		{
 			Ray.TMin = 0.01;
-			Ray.TMax = 1e20;
+			Ray.TMax = POSITIVE_INFINITY;
 		
 			Payload = TraceMaterialRay(
 				TLAS,
