GPULM: Adopt GPU instance culling changes (InstanceSceneDataOverride)

diff --git a/Engine/Plugins/Experimental/GPULightmass/Source/GPULightmass/Private/LightmapRenderer.cpp b/Engine/Plugins/Experimental/GPULightmass/Source/GPULightmass/Private/LightmapRenderer.cpp
--- a/Engine/Plugins/Experimental/GPULightmass/Source/GPULightmass/Private/LightmapRenderer.cpp
+++ b/Engine/Plugins/Experimental/GPULightmass/Source/GPULightmass/Private/LightmapRenderer.cpp
@@ -430,8 +430,10 @@ void FCachedRayTracingSceneData::SetupViewUniformBufferFromSceneRenderState(FSce
 {
 	TResourceArray<FPrimitiveSceneShaderData> PrimitiveSceneData;
 	TResourceArray<FLightmapSceneShaderData> LightmapSceneData;
+	TResourceArray<FInstanceSceneShaderData> InstanceSceneData;
 
 	PrimitiveSceneData.AddZeroed(Scene.StaticMeshInstanceRenderStates.Elements.Num());
+	InstanceSceneData.AddZeroed(Scene.StaticMeshInstanceRenderStates.Elements.Num());
 
 	TArray<int32> LightmapSceneDataStartOffsets;
 	LightmapSceneDataStartOffsets.AddZeroed(Scene.StaticMeshInstanceRenderStates.Elements.Num());
@@ -448,6 +450,13 @@ void FCachedRayTracingSceneData::SetupViewUniformBufferFromSceneRenderState(FSce
 			ConservativeLightmapEntriesNum += Instance.LODLightmapRenderStates.Num();
 		}
 
+		for (int32 InstanceGroupIndex = 0; InstanceGroupIndex < Scene.InstanceGroupRenderStates.Elements.Num(); InstanceGroupIndex++)
+		{
+			FInstanceGroupRenderState& InstanceGroup = Scene.InstanceGroupRenderStates.Elements[InstanceGroupIndex];
+			LightmapSceneDataStartOffsets[InstanceGroupIndex] = ConservativeLightmapEntriesNum;
+			ConservativeLightmapEntriesNum += InstanceGroup.LODLightmapRenderStates.Num();
+		}
+
 		LightmapSceneData.AddZeroed(ConservativeLightmapEntriesNum);
 	}
 
@@ -460,7 +469,8 @@ void FCachedRayTracingSceneData::SetupViewUniformBufferFromSceneRenderState(FSce
 				&Scene,
 				&PrimitiveSceneData,
 				&LightmapSceneData,
-				&LightmapSceneDataStartOffsets
+				&LightmapSceneDataStartOffsets,
+				&InstanceSceneData
 			](int32 InstanceIndex = 0)
 			{
 				FStaticMeshInstanceRenderState& Instance = Scene.StaticMeshInstanceRenderStates.Elements[InstanceIndex];
@@ -468,7 +478,23 @@ void FCachedRayTracingSceneData::SetupViewUniformBufferFromSceneRenderState(FSce
 				FPrimitiveUniformShaderParameters PrimitiveUniformShaderParameters = Instance.PrimitiveUniformShaderParameters;
 				PrimitiveUniformShaderParameters.LightmapDataIndex = LightmapSceneDataStartOffsets[InstanceIndex];
 				PrimitiveUniformShaderParameters.LightmapUVIndex = 0; // TODO: LightmapUVIndex
+				PrimitiveUniformShaderParameters.InstanceDataOffset = InstanceIndex;
+				PrimitiveUniformShaderParameters.NumInstanceDataEntries = 1;
 				PrimitiveSceneData[InstanceIndex] = FPrimitiveSceneShaderData(PrimitiveUniformShaderParameters);
+				InstanceSceneData[InstanceIndex] = FInstanceSceneShaderData(
+					GetInstanceUniformShaderParameters(
+						PrimitiveUniformShaderParameters.LocalToWorld, 
+						PrimitiveUniformShaderParameters.PreviousLocalToWorld,
+						(PrimitiveUniformShaderParameters.LocalObjectBoundsMin + PrimitiveUniformShaderParameters.LocalObjectBoundsMax) * 0.5f,
+						PrimitiveUniformShaderParameters.LocalObjectBoundsMax - PrimitiveUniformShaderParameters.LocalObjectBoundsMin,
+						PrimitiveUniformShaderParameters.NonUniformScale, PrimitiveUniformShaderParameters.InvNonUniformScaleAndDeterminantSign,
+						FVector4(ForceInitToZero),
+						FNaniteInfo(),
+						InstanceIndex,
+						0xFFFFFFFFu,
+						0.0f,
+						true
+					));
 
 				for (int32 LODIndex = 0; LODIndex < Instance.LODLightmapRenderStates.Num(); LODIndex++)
 				{
@@ -489,6 +515,62 @@ void FCachedRayTracingSceneData::SetupViewUniformBufferFromSceneRenderState(FSce
 					LightmapSceneData[LightmapSceneDataStartOffsets[InstanceIndex] + LODIndex] = FLightmapSceneShaderData(LightmapParams);
 				}
 			});
+
+		for (int32 InstanceGroupIndex = 0; InstanceGroupIndex < Scene.InstanceGroupRenderStates.Elements.Num(); InstanceGroupIndex++)
+		{
+			FInstanceGroupRenderState& InstanceGroup = Scene.InstanceGroupRenderStates.Elements[InstanceGroupIndex];
+
+			int32 NumInstancesThisGroup = (int32)InstanceGroup.InstancedRenderData->PerInstanceRenderData->InstanceBuffer.GetNumInstances();
+
+			FPrimitiveUniformShaderParameters PrimitiveUniformShaderParameters = GetPrimitiveUniformShaderParameters(
+				InstanceGroup.LocalToWorld,
+				InstanceGroup.LocalToWorld,
+				InstanceGroup.ActorPosition,
+				InstanceGroup.WorldBounds,
+				InstanceGroup.LocalBounds,
+				InstanceGroup.LocalBounds,
+				false,
+				false,
+				false,
+				false,
+				false,
+				false,
+				0b111,
+				LightmapSceneDataStartOffsets[Scene.StaticMeshInstanceRenderStates.Elements.Num() + InstanceGroupIndex],
+				INDEX_NONE,
+				INDEX_NONE,
+				/* bOutputVelocity = */ false,
+				nullptr,
+				/* bCastContactShadow = */ true,
+				InstanceSceneData.Num(),
+				NumInstancesThisGroup,
+				/* bCastShadow = */ true
+			);
+
+			for (int32 InstanceIdx = 0; InstanceIdx < NumInstancesThisGroup; InstanceIdx++)
+			{
+				FMatrix Transform;
+				InstanceGroup.InstancedRenderData->PerInstanceRenderData->InstanceBuffer.GetInstanceTransform(InstanceIdx, Transform);
+				Transform.M[3][3] = 1.0f;
+				FMatrix InstanceTransform = Transform * InstanceGroup.LocalToWorld;
+
+				FPrimitiveInstance Instance;
+				Instance.PrimitiveId = PrimitiveSceneData.Num();
+				Instance.InstanceToLocal = Transform;
+				// GPUCULL_TODO: not sure this is needed either - might be better to delegate to later anyway since inverse can then be threaded, plus some platforms might not need it at all.
+				Instance.LocalToInstance = Instance.InstanceToLocal.Inverse();
+				// Filled in during GPU Scene update...
+				Instance.LocalToWorld = InstanceTransform;
+				Instance.PrevLocalToWorld = InstanceTransform;
+				Instance.RenderBounds = InstanceGroup.RenderData->Bounds;
+				Instance.LocalBounds = Instance.RenderBounds.TransformBy(Instance.InstanceToLocal);
+
+				InstanceSceneData.Add(FInstanceSceneShaderData(Instance));
+			}
+
+			InstanceGroup.UniformBuffer = TUniformBufferRef<FPrimitiveUniformShaderParameters>::CreateUniformBufferImmediate(PrimitiveUniformShaderParameters, UniformBuffer_MultiFrame);
+			PrimitiveSceneData.Add(FPrimitiveSceneShaderData(PrimitiveUniformShaderParameters));
+		}
 	}
 
 	{
@@ -520,6 +602,30 @@ void FCachedRayTracingSceneData::SetupViewUniformBufferFromSceneRenderState(FSce
 			LightmapSceneDataBufferSRV = RHICreateShaderResourceView(LightmapSceneDataBufferRHI);
 		}
 
+		{
+			TRACE_CPUPROFILER_EVENT_SCOPE(InstanceSceneData);
+
+			TResourceArray<FVector4> InstanceSceneDataSOA;
+			InstanceSceneDataSOA.AddZeroed(FInstanceSceneShaderData::InstanceDataStrideInFloat4s * InstanceSceneData.Num());
+			for (int32 ArrayIndex = 0; ArrayIndex < FInstanceSceneShaderData::InstanceDataStrideInFloat4s; ArrayIndex++)
+			{
+				for (int32 DataIndex = 0; DataIndex < InstanceSceneData.Num(); DataIndex++)
+				{
+					InstanceSceneDataSOA[ArrayIndex * InstanceSceneData.Num() + DataIndex] = InstanceSceneData[DataIndex].Data[ArrayIndex];
+				}
+			}
+
+			FRHIResourceCreateInfo CreateInfo(TEXT("InstanceSceneDataBuffer"), &InstanceSceneDataSOA);
+			if (InstanceSceneDataSOA.GetResourceDataSize() == 0)
+			{
+				InstanceSceneDataSOA.AddZeroed(FInstanceSceneShaderData::InstanceDataStrideInFloat4s);
+			}
+
+			InstanceSceneDataBufferRHI = RHICreateStructuredBuffer(sizeof(FVector4), InstanceSceneDataSOA.GetResourceDataSize(), BUF_Static | BUF_ShaderResource, CreateInfo);
+			InstanceSceneDataBufferSRV = RHICreateShaderResourceView(InstanceSceneDataBufferRHI);
+			InstanceSceneDataSOAStride = InstanceSceneData.Num();
+		}
+
 		FViewUniformShaderParameters ViewUniformBufferParameters;
 		CachedViewUniformBuffer = TUniformBufferRef<FViewUniformShaderParameters>::CreateUniformBufferImmediate(ViewUniformBufferParameters, UniformBuffer_MultiFrame, EUniformBufferValidation::None);
 	}
@@ -702,6 +808,7 @@ void FSceneRenderState::SetupRayTracingScene()
 
 		View.PrimitiveSceneDataOverrideSRV = CachedRayTracingScene->PrimitiveSceneDataBufferSRV;
 		View.LightmapSceneDataOverrideSRV = CachedRayTracingScene->LightmapSceneDataBufferSRV;
+		View.InstanceSceneDataOverrideSRV = CachedRayTracingScene->InstanceSceneDataBufferSRV;
 
 		{
 			TRACE_CPUPROFILER_EVENT_SCOPE(SetupUniformBufferParameters);
@@ -715,6 +822,8 @@ void FSceneRenderState::SetupRayTracingScene()
 				TVC_MAX,
 				*View.CachedViewUniformShaderParameters);
 
+			View.CachedViewUniformShaderParameters->InstanceDataSOAStride = CachedRayTracingScene->InstanceSceneDataSOAStride;
+
 			if (LightSceneRenderState.SkyLight.IsSet())
 			{
 				View.CachedViewUniformShaderParameters->SkyIrradianceEnvironmentMap = LightSceneRenderState.SkyLight->SkyIrradianceEnvironmentMap.SRV;
diff --git a/Engine/Plugins/Experimental/GPULightmass/Source/GPULightmass/Private/Scene/Scene.h b/Engine/Plugins/Experimental/GPULightmass/Source/GPULightmass/Private/Scene/Scene.h
--- a/Engine/Plugins/Experimental/GPULightmass/Source/GPULightmass/Private/Scene/Scene.h
+++ b/Engine/Plugins/Experimental/GPULightmass/Source/GPULightmass/Private/Scene/Scene.h
@@ -66,6 +66,10 @@ struct FCachedRayTracingSceneData
 	FBufferRHIRef LightmapSceneDataBufferRHI;
 	FShaderResourceViewRHIRef LightmapSceneDataBufferSRV;
 
+	FBufferRHIRef InstanceSceneDataBufferRHI;
+	FShaderResourceViewRHIRef InstanceSceneDataBufferSRV;
+	uint32 InstanceSceneDataSOAStride;
+
 	TArray<FRayTracingGeometryInstance> RayTracingGeometryInstances;
 
 	TUniformBufferRef<FViewUniformShaderParameters> CachedViewUniformBuffer;
