GPULM: Fixes for global uniform buffer binding changes

diff --git a/Engine/Plugins/Experimental/GPULightmass/Source/GPULightmass/Private/VolumetricLightmap.cpp b/Engine/Plugins/Experimental/GPULightmass/Source/GPULightmass/Private/VolumetricLightmap.cpp
--- a/Engine/Plugins/Experimental/GPULightmass/Source/GPULightmass/Private/VolumetricLightmap.cpp
+++ b/Engine/Plugins/Experimental/GPULightmass/Source/GPULightmass/Private/VolumetricLightmap.cpp
@@ -23,7 +23,7 @@ IMPLEMENT_GLOBAL_SHADER(FSplatVolumeCS, "/Plugin/GPULightmass/Private/BrickAlloc
 IMPLEMENT_GLOBAL_SHADER(FStitchBorderCS, "/Plugin/GPULightmass/Private/BrickAllocationManagement.usf", "StitchBorderCS", SF_Compute);
 IMPLEMENT_GLOBAL_SHADER(FFinalizeBrickResultsCS, "/Plugin/GPULightmass/Private/BrickAllocationManagement.usf", "FinalizeBrickResultsCS", SF_Compute);
 
-IMPLEMENT_GLOBAL_SHADER_PARAMETER_STRUCT(FVLMVoxelizationParams, "VLMVoxelizationParams");
+IMPLEMENT_STATIC_UNIFORM_BUFFER_STRUCT(FVLMVoxelizationParams, "VLMVoxelizationParams", SceneTextures);
 
 void InitializeBrickData(FIntVector BrickDataDimensions, FVolumetricLightmapBrickData& BrickData, const bool bForAccumulation)
 {
@@ -186,6 +186,10 @@ void FVolumetricLightmapRenderer::VoxelizeScene()
 
 	FRHICommandListImmediate& RHICmdList = FRHICommandListExecutor::GetImmediateCommandList();
 
+	FUniformBufferStaticBindings GlobalUniformBuffers;
+	GlobalUniformBuffers.AddUniformBuffer(PassUniformBuffer);
+	RHICmdList.SetGlobalUniformBuffers(GlobalUniformBuffers);
+
 	for (int32 MipLevel = 0; MipLevel < VoxelizationVolumeMips.Num(); MipLevel++)
 	{
 		FGlobalShaderMap* GlobalShaderMap = GetGlobalShaderMap(GMaxRHIFeatureLevel);
@@ -483,15 +487,16 @@ void FVolumetricLightmapRenderer::BackgroundTick()
 	}
 
 	FRDGBuilder GraphBuilder(FRHICommandListExecutor::GetImmediateCommandList());
-	RDG_EVENT_SCOPE(GraphBuilder, "Volumetric Lightmap Path Tracing");
+	{
+		RDG_EVENT_SCOPE(GraphBuilder, "Volumetric Lightmap Path Tracing");
 
-	bool bLastFewFramesIdle = !GCurrentLevelEditingViewportClient || !GCurrentLevelEditingViewportClient->IsRealtime();
+		bool bLastFewFramesIdle = !GCurrentLevelEditingViewportClient || !GCurrentLevelEditingViewportClient->IsRealtime();
 
-	int32 NumSamplesThisFrame = !bLastFewFramesIdle ? 1 : 32;
+		int32 NumSamplesThisFrame = !bLastFewFramesIdle ? 1 : 32;
 
-	// manually handle transitions since the buffers are not (yet) managed by RDG
-	GraphBuilder.AddPass(RDG_EVENT_NAME("Transition Buffers"), ERDGPassFlags::None, 
-		[this](FRHICommandListImmediate& RHICmdList) {
+		// manually handle transitions since the buffers are not (yet) managed by RDG
+		GraphBuilder.AddPass(RDG_EVENT_NAME("Transition Buffers"), ERDGPassFlags::None,
+			[this](FRHICommandListImmediate& RHICmdList) {
 			RHICmdList.Transition({
 				FRHITransitionInfo(VolumetricLightmapData.BrickData.AmbientVector.UAV, ERHIAccess::Unknown, ERHIAccess::UAVCompute),
 				FRHITransitionInfo(VolumetricLightmapData.BrickData.SHCoefficients[0].UAV, ERHIAccess::Unknown, ERHIAccess::UAVCompute),
@@ -503,127 +508,86 @@ void FVolumetricLightmapRenderer::BackgroundTick()
 				}
 			);
 		}
-	);
-
-	for (int32 SampleIndex = 0; SampleIndex < NumSamplesThisFrame; SampleIndex++)
-	{
-		int32 MaxBricksPerFrame = FMath::Min(512, NumTotalBricks);
-		int32 NumFramesOneRound = FMath::DivideAndRoundUp(NumTotalBricks, MaxBricksPerFrame);
-		int32 BrickBatchOffset = MaxBricksPerFrame * (FrameNumber % NumFramesOneRound);
-		int32 BricksToCalcThisFrame = FMath::Min(MaxBricksPerFrame, NumTotalBricks - BrickBatchOffset);
-		if (BricksToCalcThisFrame <= 0) continue;
+		);
 
-#if RHI_RAYTRACING
-		if (IsRayTracingEnabled())
+		for (int32 SampleIndex = 0; SampleIndex < NumSamplesThisFrame; SampleIndex++)
 		{
-			// These two buffers must have lifetime extended beyond RHICmdList.RayTraceDispatch()
-			TUniformBufferRef<FSkyLightData> SkyLightDataUniformBuffer;
-
-			FGlobalShaderMap* GlobalShaderMap = GetGlobalShaderMap(GMaxRHIFeatureLevel);
+			int32 MaxBricksPerFrame = FMath::Min(512, NumTotalBricks);
+			int32 NumFramesOneRound = FMath::DivideAndRoundUp(NumTotalBricks, MaxBricksPerFrame);
+			int32 BrickBatchOffset = MaxBricksPerFrame * (FrameNumber % NumFramesOneRound);
+			int32 BricksToCalcThisFrame = FMath::Min(MaxBricksPerFrame, NumTotalBricks - BrickBatchOffset);
+			if (BricksToCalcThisFrame <= 0) continue;
 
+#if RHI_RAYTRACING
+			if (IsRayTracingEnabled())
 			{
-				FVolumetricLightmapPathTracingRGS::FPermutationDomain PermutationVector;
-				PermutationVector.Set<FVolumetricLightmapPathTracingRGS::FUseIrradianceCaching>(Scene->Settings->bUseIrradianceCaching);
-				TShaderRef<FVolumetricLightmapPathTracingRGS> RayGenShader = GlobalShaderMap->GetShader<FVolumetricLightmapPathTracingRGS>(PermutationVector);
-
-				FVolumetricLightmapPathTracingRGS::FParameters* PassParameters = GraphBuilder.AllocParameters<FVolumetricLightmapPathTracingRGS::FParameters>();
-				PassParameters->FrameNumber = FrameNumber / NumFramesOneRound;
-				PassParameters->VolumeMin = VolumeMin;
-				PassParameters->VolumeSize = VolumeSize;
-				PassParameters->IndirectionTextureDim = IndirectionTextureDimensions;
-				PassParameters->TLAS = Scene->RayTracingScene->GetShaderResourceView();
-				PassParameters->BrickRequests = BrickRequests.SRV;
-				PassParameters->NumTotalBricks = NumTotalBricks;
-				PassParameters->BrickBatchOffset = BrickBatchOffset;
-				PassParameters->AmbientVector = AccumulationBrickData.AmbientVector.UAV;
-				PassParameters->SHCoefficients0R = AccumulationBrickData.SHCoefficients[0].UAV;
-				PassParameters->SHCoefficients1R = AccumulationBrickData.SHCoefficients[1].UAV;
-				PassParameters->SHCoefficients0G = AccumulationBrickData.SHCoefficients[2].UAV;
-				PassParameters->SHCoefficients1G = AccumulationBrickData.SHCoefficients[3].UAV;
-				PassParameters->SHCoefficients0B = AccumulationBrickData.SHCoefficients[4].UAV;
-				PassParameters->SHCoefficients1B = AccumulationBrickData.SHCoefficients[5].UAV;
-				PassParameters->OutAmbientVector = VolumetricLightmapData.BrickData.AmbientVector.UAV;
-				PassParameters->OutSHCoefficients0R = VolumetricLightmapData.BrickData.SHCoefficients[0].UAV;
-				PassParameters->OutSHCoefficients1R = VolumetricLightmapData.BrickData.SHCoefficients[1].UAV;
-				PassParameters->OutSHCoefficients0G = VolumetricLightmapData.BrickData.SHCoefficients[2].UAV;
-				PassParameters->OutSHCoefficients1G = VolumetricLightmapData.BrickData.SHCoefficients[3].UAV;
-				PassParameters->OutSHCoefficients0B = VolumetricLightmapData.BrickData.SHCoefficients[4].UAV;
-				PassParameters->OutSHCoefficients1B = VolumetricLightmapData.BrickData.SHCoefficients[5].UAV;
-				PassParameters->ViewUniformBuffer = Scene->ReferenceView->ViewUniformBuffer;
-				PassParameters->IrradianceCachingParameters = Scene->IrradianceCache->IrradianceCachingParametersUniformBuffer;
-
-				SetupPathTracingLightParameters(Scene->LightSceneRenderState, GraphBuilder, &PassParameters->SceneLights, &PassParameters->SceneLightCount);
+				// These two buffers must have lifetime extended beyond RHICmdList.RayTraceDispatch()
+				TUniformBufferRef<FSkyLightData> SkyLightDataUniformBuffer;
 
-				PassParameters->SkyLight = CreateUniformBufferImmediate(SetupSkyLightParameters(Scene->LightSceneRenderState), EUniformBufferUsage::UniformBuffer_SingleFrame);
-
-				// TODO: find a way to share IES atlas with path tracer ...
-				PassParameters->IESTexture = GWhiteTexture->TextureRHI;
-				PassParameters->IESTextureSampler = TStaticSamplerState<SF_Bilinear, AM_Clamp, AM_Clamp, AM_Clamp>::GetRHI();
+				FGlobalShaderMap* GlobalShaderMap = GetGlobalShaderMap(GMaxRHIFeatureLevel);
 
-				FSceneRenderState* SceneRenderState = Scene; // capture member variable
-				GraphBuilder.AddPass(
-					RDG_EVENT_NAME("VolumetricLightmapPathTracing %d bricks %d rays", BricksToCalcThisFrame, BricksToCalcThisFrame * (BrickSize + 1) * (BrickSize + 1) * (BrickSize + 1)),
-					PassParameters,
-					ERDGPassFlags::Compute,
-					[PassParameters, RayGenShader, SceneRenderState, BricksToCalcThisFrame](FRHICommandListImmediate& RHICmdList)
+				{
+					FVolumetricLightmapPathTracingRGS::FPermutationDomain PermutationVector;
+					PermutationVector.Set<FVolumetricLightmapPathTracingRGS::FUseIrradianceCaching>(Scene->Settings->bUseIrradianceCaching);
+					TShaderRef<FVolumetricLightmapPathTracingRGS> RayGenShader = GlobalShaderMap->GetShader<FVolumetricLightmapPathTracingRGS>(PermutationVector);
+
+					FVolumetricLightmapPathTracingRGS::FParameters* PassParameters = GraphBuilder.AllocParameters<FVolumetricLightmapPathTracingRGS::FParameters>();
+					PassParameters->FrameNumber = FrameNumber / NumFramesOneRound;
+					PassParameters->VolumeMin = VolumeMin;
+					PassParameters->VolumeSize = VolumeSize;
+					PassParameters->IndirectionTextureDim = IndirectionTextureDimensions;
+					PassParameters->TLAS = Scene->RayTracingScene->GetShaderResourceView();
+					PassParameters->BrickRequests = BrickRequests.SRV;
+					PassParameters->NumTotalBricks = NumTotalBricks;
+					PassParameters->BrickBatchOffset = BrickBatchOffset;
+					PassParameters->AmbientVector = AccumulationBrickData.AmbientVector.UAV;
+					PassParameters->SHCoefficients0R = AccumulationBrickData.SHCoefficients[0].UAV;
+					PassParameters->SHCoefficients1R = AccumulationBrickData.SHCoefficients[1].UAV;
+					PassParameters->SHCoefficients0G = AccumulationBrickData.SHCoefficients[2].UAV;
+					PassParameters->SHCoefficients1G = AccumulationBrickData.SHCoefficients[3].UAV;
+					PassParameters->SHCoefficients0B = AccumulationBrickData.SHCoefficients[4].UAV;
+					PassParameters->SHCoefficients1B = AccumulationBrickData.SHCoefficients[5].UAV;
+					PassParameters->OutAmbientVector = VolumetricLightmapData.BrickData.AmbientVector.UAV;
+					PassParameters->OutSHCoefficients0R = VolumetricLightmapData.BrickData.SHCoefficients[0].UAV;
+					PassParameters->OutSHCoefficients1R = VolumetricLightmapData.BrickData.SHCoefficients[1].UAV;
+					PassParameters->OutSHCoefficients0G = VolumetricLightmapData.BrickData.SHCoefficients[2].UAV;
+					PassParameters->OutSHCoefficients1G = VolumetricLightmapData.BrickData.SHCoefficients[3].UAV;
+					PassParameters->OutSHCoefficients0B = VolumetricLightmapData.BrickData.SHCoefficients[4].UAV;
+					PassParameters->OutSHCoefficients1B = VolumetricLightmapData.BrickData.SHCoefficients[5].UAV;
+					PassParameters->ViewUniformBuffer = Scene->ReferenceView->ViewUniformBuffer;
+					PassParameters->IrradianceCachingParameters = Scene->IrradianceCache->IrradianceCachingParametersUniformBuffer;
+
+					SetupPathTracingLightParameters(Scene->LightSceneRenderState, GraphBuilder, &PassParameters->SceneLights, &PassParameters->SceneLightCount);
+
+					PassParameters->SkyLight = CreateUniformBufferImmediate(SetupSkyLightParameters(Scene->LightSceneRenderState), EUniformBufferUsage::UniformBuffer_SingleFrame);
+
+					// TODO: find a way to share IES atlas with path tracer ...
+					PassParameters->IESTexture = GWhiteTexture->TextureRHI;
+					PassParameters->IESTextureSampler = TStaticSamplerState<SF_Bilinear, AM_Clamp, AM_Clamp, AM_Clamp>::GetRHI();
+
+					FSceneRenderState* SceneRenderState = Scene; // capture member variable
+					GraphBuilder.AddPass(
+						RDG_EVENT_NAME("VolumetricLightmapPathTracing %d bricks %d rays", BricksToCalcThisFrame, BricksToCalcThisFrame * (BrickSize + 1) * (BrickSize + 1) * (BrickSize + 1)),
+						PassParameters,
+						ERDGPassFlags::Compute,
+						[PassParameters, RayGenShader, SceneRenderState, BricksToCalcThisFrame](FRHICommandListImmediate& RHICmdList)
 					{
 						FRayTracingShaderBindingsWriter GlobalResources;
 						SetShaderParameters(GlobalResources, RayGenShader, *PassParameters);
 
 						FRHIRayTracingScene* RayTracingSceneRHI = SceneRenderState->RayTracingScene;
-						RHICmdList.RayTraceDispatch(SceneRenderState->RayTracingPipelineState, RayGenShader.GetRayTracingShader(), RayTracingSceneRHI, GlobalResources, BricksToCalcThisFrame * (BrickSize + 1)* (BrickSize + 1)* (BrickSize + 1), 1);
+						RHICmdList.RayTraceDispatch(SceneRenderState->RayTracingPipelineState, RayGenShader.GetRayTracingShader(), RayTracingSceneRHI, GlobalResources, BricksToCalcThisFrame * (BrickSize + 1) * (BrickSize + 1) * (BrickSize + 1), 1);
 					}
-				);
+					);
 
 
+				}
 			}
-		}
 #endif
-		{
-			// manually handle transitions since the buffers are not (yet) managed by RDG
-			GraphBuilder.AddPass(RDG_EVENT_NAME("Transition Buffers"), ERDGPassFlags::None,
-				[this](FRHICommandListImmediate& RHICmdList) {
-					RHICmdList.Transition({
-						FRHITransitionInfo(AccumulationBrickData.AmbientVector.UAV, ERHIAccess::UAVCompute, ERHIAccess::UAVCompute),
-						FRHITransitionInfo(AccumulationBrickData.SHCoefficients[0].UAV, ERHIAccess::UAVCompute, ERHIAccess::UAVCompute),
-						FRHITransitionInfo(AccumulationBrickData.SHCoefficients[1].UAV, ERHIAccess::UAVCompute, ERHIAccess::UAVCompute),
-						FRHITransitionInfo(AccumulationBrickData.SHCoefficients[2].UAV, ERHIAccess::UAVCompute, ERHIAccess::UAVCompute),
-						FRHITransitionInfo(AccumulationBrickData.SHCoefficients[3].UAV, ERHIAccess::UAVCompute, ERHIAccess::UAVCompute),
-						FRHITransitionInfo(AccumulationBrickData.SHCoefficients[4].UAV, ERHIAccess::UAVCompute, ERHIAccess::UAVCompute),
-						FRHITransitionInfo(AccumulationBrickData.SHCoefficients[5].UAV, ERHIAccess::UAVCompute, ERHIAccess::UAVCompute),
-						}
-					);
-				}
-			);
-
-			FGlobalShaderMap* GlobalShaderMap = GetGlobalShaderMap(GMaxRHIFeatureLevel);
-
-			TShaderMapRef<FFinalizeBrickResultsCS> ComputeShader(GlobalShaderMap);
-
-			FFinalizeBrickResultsCS::FParameters* PassParameters = GraphBuilder.AllocParameters<FFinalizeBrickResultsCS::FParameters>();
-			PassParameters->NumTotalBricks = NumTotalBricks;
-			PassParameters->BrickBatchOffset = BrickBatchOffset;
-			PassParameters->BrickRequests = BrickRequests.UAV;
-			PassParameters->AmbientVector = AccumulationBrickData.AmbientVector.Texture;
-			PassParameters->SHCoefficients0R = AccumulationBrickData.SHCoefficients[0].Texture;
-			PassParameters->SHCoefficients1R = AccumulationBrickData.SHCoefficients[1].Texture;
-			PassParameters->SHCoefficients0G = AccumulationBrickData.SHCoefficients[2].Texture;
-			PassParameters->SHCoefficients1G = AccumulationBrickData.SHCoefficients[3].Texture;
-			PassParameters->SHCoefficients0B = AccumulationBrickData.SHCoefficients[4].Texture;
-			PassParameters->SHCoefficients1B = AccumulationBrickData.SHCoefficients[5].Texture;
-			PassParameters->OutAmbientVector = VolumetricLightmapData.BrickData.AmbientVector.UAV;
-			PassParameters->OutSHCoefficients0R = VolumetricLightmapData.BrickData.SHCoefficients[0].UAV;
-			PassParameters->OutSHCoefficients1R = VolumetricLightmapData.BrickData.SHCoefficients[1].UAV;
-			PassParameters->OutSHCoefficients0G = VolumetricLightmapData.BrickData.SHCoefficients[2].UAV;
-			PassParameters->OutSHCoefficients1G = VolumetricLightmapData.BrickData.SHCoefficients[3].UAV;
-			PassParameters->OutSHCoefficients0B = VolumetricLightmapData.BrickData.SHCoefficients[4].UAV;
-			PassParameters->OutSHCoefficients1B = VolumetricLightmapData.BrickData.SHCoefficients[5].UAV;
-
-			FComputeShaderUtils::AddPass(GraphBuilder, RDG_EVENT_NAME("FinalizeBrickResults"), ComputeShader, PassParameters, FIntVector(BricksToCalcThisFrame, 1, 1));
-
-			// manually handle transitions since the buffers are not (yet) managed by RDG
-			GraphBuilder.AddPass(RDG_EVENT_NAME("Transition Buffers"), ERDGPassFlags::None,
-				[this](FRHICommandListImmediate& RHICmdList) {
+			{
+				// manually handle transitions since the buffers are not (yet) managed by RDG
+				GraphBuilder.AddPass(RDG_EVENT_NAME("Transition Buffers"), ERDGPassFlags::None,
+					[this](FRHICommandListImmediate& RHICmdList) {
 					RHICmdList.Transition({
 						FRHITransitionInfo(AccumulationBrickData.AmbientVector.UAV, ERHIAccess::UAVCompute, ERHIAccess::UAVCompute),
 						FRHITransitionInfo(AccumulationBrickData.SHCoefficients[0].UAV, ERHIAccess::UAVCompute, ERHIAccess::UAVCompute),
@@ -635,26 +599,23 @@ void FVolumetricLightmapRenderer::BackgroundTick()
 						}
 					);
 				}
-			);
-		}
+				);
 
-		{
-			// Doing 2 passes no longer makes sense in an amortized setup
-			// for (int32 StitchPass = 0; StitchPass < 2; StitchPass++)
-			{
 				FGlobalShaderMap* GlobalShaderMap = GetGlobalShaderMap(GMaxRHIFeatureLevel);
 
-				TShaderMapRef<FStitchBorderCS> ComputeShader(GlobalShaderMap);
+				TShaderMapRef<FFinalizeBrickResultsCS> ComputeShader(GlobalShaderMap);
 
-				FStitchBorderCS::FParameters* PassParameters = GraphBuilder.AllocParameters<FStitchBorderCS::FParameters>();
-				PassParameters->BrickDataDimensions = VolumetricLightmapData.BrickDataDimensions;
-				PassParameters->IndirectionTextureDim = IndirectionTextureDimensions;
-				PassParameters->FrameNumber = FrameNumber / NumFramesOneRound;
+				FFinalizeBrickResultsCS::FParameters* PassParameters = GraphBuilder.AllocParameters<FFinalizeBrickResultsCS::FParameters>();
 				PassParameters->NumTotalBricks = NumTotalBricks;
 				PassParameters->BrickBatchOffset = BrickBatchOffset;
-				PassParameters->IndirectionTexture = IndirectionTexture->GetRenderTargetItem().UAV;
 				PassParameters->BrickRequests = BrickRequests.UAV;
 				PassParameters->AmbientVector = AccumulationBrickData.AmbientVector.Texture;
+				PassParameters->SHCoefficients0R = AccumulationBrickData.SHCoefficients[0].Texture;
+				PassParameters->SHCoefficients1R = AccumulationBrickData.SHCoefficients[1].Texture;
+				PassParameters->SHCoefficients0G = AccumulationBrickData.SHCoefficients[2].Texture;
+				PassParameters->SHCoefficients1G = AccumulationBrickData.SHCoefficients[3].Texture;
+				PassParameters->SHCoefficients0B = AccumulationBrickData.SHCoefficients[4].Texture;
+				PassParameters->SHCoefficients1B = AccumulationBrickData.SHCoefficients[5].Texture;
 				PassParameters->OutAmbientVector = VolumetricLightmapData.BrickData.AmbientVector.UAV;
 				PassParameters->OutSHCoefficients0R = VolumetricLightmapData.BrickData.SHCoefficients[0].UAV;
 				PassParameters->OutSHCoefficients1R = VolumetricLightmapData.BrickData.SHCoefficients[1].UAV;
@@ -663,11 +624,55 @@ void FVolumetricLightmapRenderer::BackgroundTick()
 				PassParameters->OutSHCoefficients0B = VolumetricLightmapData.BrickData.SHCoefficients[4].UAV;
 				PassParameters->OutSHCoefficients1B = VolumetricLightmapData.BrickData.SHCoefficients[5].UAV;
 
-				FComputeShaderUtils::AddPass(GraphBuilder, RDG_EVENT_NAME("VolumetricLightmapStitching %d bricks", BricksToCalcThisFrame), ComputeShader, PassParameters, FIntVector(BricksToCalcThisFrame, 1, 1));
+				FComputeShaderUtils::AddPass(GraphBuilder, RDG_EVENT_NAME("FinalizeBrickResults"), ComputeShader, PassParameters, FIntVector(BricksToCalcThisFrame, 1, 1));
 
 				// manually handle transitions since the buffers are not (yet) managed by RDG
 				GraphBuilder.AddPass(RDG_EVENT_NAME("Transition Buffers"), ERDGPassFlags::None,
 					[this](FRHICommandListImmediate& RHICmdList) {
+					RHICmdList.Transition({
+						FRHITransitionInfo(AccumulationBrickData.AmbientVector.UAV, ERHIAccess::UAVCompute, ERHIAccess::UAVCompute),
+						FRHITransitionInfo(AccumulationBrickData.SHCoefficients[0].UAV, ERHIAccess::UAVCompute, ERHIAccess::UAVCompute),
+						FRHITransitionInfo(AccumulationBrickData.SHCoefficients[1].UAV, ERHIAccess::UAVCompute, ERHIAccess::UAVCompute),
+						FRHITransitionInfo(AccumulationBrickData.SHCoefficients[2].UAV, ERHIAccess::UAVCompute, ERHIAccess::UAVCompute),
+						FRHITransitionInfo(AccumulationBrickData.SHCoefficients[3].UAV, ERHIAccess::UAVCompute, ERHIAccess::UAVCompute),
+						FRHITransitionInfo(AccumulationBrickData.SHCoefficients[4].UAV, ERHIAccess::UAVCompute, ERHIAccess::UAVCompute),
+						FRHITransitionInfo(AccumulationBrickData.SHCoefficients[5].UAV, ERHIAccess::UAVCompute, ERHIAccess::UAVCompute),
+						}
+					);
+				}
+				);
+			}
+
+			{
+				// Doing 2 passes no longer makes sense in an amortized setup
+				// for (int32 StitchPass = 0; StitchPass < 2; StitchPass++)
+				{
+					FGlobalShaderMap* GlobalShaderMap = GetGlobalShaderMap(GMaxRHIFeatureLevel);
+
+					TShaderMapRef<FStitchBorderCS> ComputeShader(GlobalShaderMap);
+
+					FStitchBorderCS::FParameters* PassParameters = GraphBuilder.AllocParameters<FStitchBorderCS::FParameters>();
+					PassParameters->BrickDataDimensions = VolumetricLightmapData.BrickDataDimensions;
+					PassParameters->IndirectionTextureDim = IndirectionTextureDimensions;
+					PassParameters->FrameNumber = FrameNumber / NumFramesOneRound;
+					PassParameters->NumTotalBricks = NumTotalBricks;
+					PassParameters->BrickBatchOffset = BrickBatchOffset;
+					PassParameters->IndirectionTexture = IndirectionTexture->GetRenderTargetItem().UAV;
+					PassParameters->BrickRequests = BrickRequests.UAV;
+					PassParameters->AmbientVector = AccumulationBrickData.AmbientVector.Texture;
+					PassParameters->OutAmbientVector = VolumetricLightmapData.BrickData.AmbientVector.UAV;
+					PassParameters->OutSHCoefficients0R = VolumetricLightmapData.BrickData.SHCoefficients[0].UAV;
+					PassParameters->OutSHCoefficients1R = VolumetricLightmapData.BrickData.SHCoefficients[1].UAV;
+					PassParameters->OutSHCoefficients0G = VolumetricLightmapData.BrickData.SHCoefficients[2].UAV;
+					PassParameters->OutSHCoefficients1G = VolumetricLightmapData.BrickData.SHCoefficients[3].UAV;
+					PassParameters->OutSHCoefficients0B = VolumetricLightmapData.BrickData.SHCoefficients[4].UAV;
+					PassParameters->OutSHCoefficients1B = VolumetricLightmapData.BrickData.SHCoefficients[5].UAV;
+
+					FComputeShaderUtils::AddPass(GraphBuilder, RDG_EVENT_NAME("VolumetricLightmapStitching %d bricks", BricksToCalcThisFrame), ComputeShader, PassParameters, FIntVector(BricksToCalcThisFrame, 1, 1));
+
+					// manually handle transitions since the buffers are not (yet) managed by RDG
+					GraphBuilder.AddPass(RDG_EVENT_NAME("Transition Buffers"), ERDGPassFlags::None,
+						[this](FRHICommandListImmediate& RHICmdList) {
 						RHICmdList.Transition({
 							FRHITransitionInfo(AccumulationBrickData.AmbientVector.UAV, ERHIAccess::UAVCompute, ERHIAccess::UAVCompute),
 							FRHITransitionInfo(AccumulationBrickData.SHCoefficients[0].UAV, ERHIAccess::UAVCompute, ERHIAccess::UAVCompute),
@@ -679,17 +684,18 @@ void FVolumetricLightmapRenderer::BackgroundTick()
 							}
 						);
 					}
-				);
+					);
+				}
 			}
-		}
 
-		FrameNumber++;
+			FrameNumber++;
 
-		SamplesTaken += BricksToCalcThisFrame * NumCellsPerBrick;
+			SamplesTaken += BricksToCalcThisFrame * NumCellsPerBrick;
 
-		if (SamplesTaken >= (uint64)NumTotalBricks * NumCellsPerBrick * Scene->Settings->GISamples * GetGISamplesMultiplier())
-		{
-			break;
+			if (SamplesTaken >= (uint64)NumTotalBricks * NumCellsPerBrick * Scene->Settings->GISamples * GetGISamplesMultiplier())
+			{
+				break;
+			}
 		}
 	}
 	GraphBuilder.Execute();
