Fix SetStaticUniformBuffers

diff --git a/Engine/Plugins/Experimental/GPULightmass/Source/GPULightmass/Private/VolumetricLightmap.cpp b/Engine/Plugins/Experimental/GPULightmass/Source/GPULightmass/Private/VolumetricLightmap.cpp
--- a/Engine/Plugins/Experimental/GPULightmass/Source/GPULightmass/Private/VolumetricLightmap.cpp
+++ b/Engine/Plugins/Experimental/GPULightmass/Source/GPULightmass/Private/VolumetricLightmap.cpp
@@ -188,7 +188,7 @@ void FVolumetricLightmapRenderer::VoxelizeScene()
 
 	FUniformBufferStaticBindings GlobalUniformBuffers;
 	GlobalUniformBuffers.AddUniformBuffer(PassUniformBuffer);
-	RHICmdList.SetGlobalUniformBuffers(GlobalUniformBuffers);
+	RHICmdList.SetStaticUniformBuffers(GlobalUniformBuffers);
 
 	for (int32 MipLevel = 0; MipLevel < VoxelizationVolumeMips.Num(); MipLevel++)
 	{
