Add missing include to GPU Lightmass from recent path tracing change

diff --git a/Engine/Plugins/Experimental/GPULightmass/Shaders/Private/LightmapPathTracing.usf b/Engine/Plugins/Experimental/GPULightmass/Shaders/Private/LightmapPathTracing.usf
--- a/Engine/Plugins/Experimental/GPULightmass/Shaders/Private/LightmapPathTracing.usf
+++ b/Engine/Plugins/Experimental/GPULightmass/Shaders/Private/LightmapPathTracing.usf
@@ -35,6 +35,7 @@ RWTexture2D<uint> RayGuidingLuminance;
 Texture2D<float> RayGuidingCDFX;
 Texture2D<float> RayGuidingCDFY;
 
+#include "/Engine/Private/ShadingModels.ush"
 #include "/Engine/Private/PathTracing/Utilities/PathTracingRandomSequence.ush"
 #include "/Engine/Private/PathTracing/Light/PathTracingLightSampling.ush"
 #include "/Engine/Private/PathTracing/Material/PathTracingMaterialSampling.ush"
