Fixed compile errors in GPULightmass due to ray tracing TLAS memory management redesign.

diff --git a/Engine/Plugins/Experimental/GPULightmass/Source/GPULightmass/Private/LightmapRenderer.cpp b/Engine/Plugins/Experimental/GPULightmass/Source/GPULightmass/Private/LightmapRenderer.cpp
--- a/Engine/Plugins/Experimental/GPULightmass/Source/GPULightmass/Private/LightmapRenderer.cpp
+++ b/Engine/Plugins/Experimental/GPULightmass/Source/GPULightmass/Private/LightmapRenderer.cpp
@@ -943,6 +943,20 @@ void FSceneRenderState::SetupRayTracingScene()
 			SCOPED_GPU_MASK(RHICmdList, FRHIGPUMask::All());
 
 			RayTracingScene = RHICreateRayTracingScene(Initializer);
+
+			uint32 NumTotalInstances = 0;
+			for (const FRayTracingGeometryInstance& Instance : View.RayTracingGeometryInstances)
+			{
+				ensure(Instance.NumTransforms >= uint32(Instance.GetTransforms().Num()));
+				NumTotalInstances += Instance.NumTransforms;
+			}
+
+			FRayTracingAccelerationStructureSize SizeInfo = RHICalcRayTracingSceneSize(NumTotalInstances, ERayTracingAccelerationStructureFlags::FastTrace);
+			FRHIResourceCreateInfo BufferCreateInfo(TEXT("LightmassRayTracingSceneBuffer"));
+			RayTracingSceneBuffer = RHICreateBuffer(uint32(SizeInfo.ResultSize), BUF_AccelerationStructure, 0, ERHIAccess::BVHWrite, BufferCreateInfo);
+			RayTracingSceneSRV = RHICreateShaderResourceView(RayTracingSceneBuffer);
+
+			RHICmdList.BindAccelerationStructureMemory(RayTracingScene, RayTracingSceneBuffer, 0);
 			RHICmdList.BuildAccelerationStructure(RayTracingScene);
 
 			FRayTracingPipelineStateInitializer PSOInitializer;
@@ -1971,7 +1985,7 @@ void FLightmapRenderer::Finalize(FRDGBuilder& GraphBuilder)
 									FLightmapPathTracingRGS::FParameters* PassParameters = GraphBuilder.AllocParameters<FLightmapPathTracingRGS::FParameters>();
 									PassParameters->LastInvalidationFrame = LastInvalidationFrame;
 									PassParameters->NumTotalSamples = Scene->Settings->GISamples;
-									PassParameters->TLAS = Scene->RayTracingScene->GetShaderResourceView();
+									PassParameters->TLAS = Scene->RayTracingSceneSRV;
 									PassParameters->GBufferWorldPosition = GBufferWorldPosition;
 									PassParameters->GBufferWorldNormal = GBufferWorldNormal;
 									PassParameters->GBufferShadingNormal = GBufferShadingNormal;
@@ -2374,7 +2388,7 @@ void FLightmapRenderer::Finalize(FRDGBuilder& GraphBuilder)
 						RayTracingResolution.Y = GPreviewLightmapPhysicalTileSize;
 
 						FStationaryLightShadowTracingRGS::FParameters* PassParameters = GraphBuilder.AllocParameters<FStationaryLightShadowTracingRGS::FParameters>();
-						PassParameters->TLAS = Scene->RayTracingScene->GetShaderResourceView();
+						PassParameters->TLAS = Scene->RayTracingSceneSRV;
 						PassParameters->BatchedTiles = GPUBatchedTileRequests.BatchedTilesSRV;
 						PassParameters->LightTypeArray = HoldReference(GraphBuilder, LightTypeSRV);
 						PassParameters->ChannelIndexArray = HoldReference(GraphBuilder, ChannelIndexSRV);
diff --git a/Engine/Plugins/Experimental/GPULightmass/Source/GPULightmass/Private/Scene/Scene.h b/Engine/Plugins/Experimental/GPULightmass/Source/GPULightmass/Private/Scene/Scene.h
--- a/Engine/Plugins/Experimental/GPULightmass/Source/GPULightmass/Private/Scene/Scene.h
+++ b/Engine/Plugins/Experimental/GPULightmass/Source/GPULightmass/Private/Scene/Scene.h
@@ -87,6 +87,9 @@ public:
 	void BackgroundTick();
 
 	FRayTracingSceneRHIRef RayTracingScene;
+	FShaderResourceViewRHIRef RayTracingSceneSRV;
+	FBufferRHIRef RayTracingSceneBuffer;
+
 	FRayTracingPipelineState* RayTracingPipelineState;
 	TSharedPtr<FViewInfo> ReferenceView;
 
diff --git a/Engine/Plugins/Experimental/GPULightmass/Source/GPULightmass/Private/VolumetricLightmap.cpp b/Engine/Plugins/Experimental/GPULightmass/Source/GPULightmass/Private/VolumetricLightmap.cpp
--- a/Engine/Plugins/Experimental/GPULightmass/Source/GPULightmass/Private/VolumetricLightmap.cpp
+++ b/Engine/Plugins/Experimental/GPULightmass/Source/GPULightmass/Private/VolumetricLightmap.cpp
@@ -536,7 +536,7 @@ void FVolumetricLightmapRenderer::BackgroundTick()
 					PassParameters->VolumeMin = VolumeMin;
 					PassParameters->VolumeSize = VolumeSize;
 					PassParameters->IndirectionTextureDim = IndirectionTextureDimensions;
-					PassParameters->TLAS = Scene->RayTracingScene->GetShaderResourceView();
+					PassParameters->TLAS = Scene->RayTracingSceneSRV;
 					PassParameters->BrickRequests = BrickRequests.SRV;
 					PassParameters->NumTotalBricks = NumTotalBricks;
 					PassParameters->BrickBatchOffset = BrickBatchOffset;
