GPULM: miscellaneous fixes

* Fix RenderPassIndex which was fetched using references
* Fix RayTracingGeometryInstances
* Irradiance Caching now uses ClearUAV instead of uploading a giant zeroed buffer

diff --git a/Engine/Plugins/Experimental/GPULightmass/Source/GPULightmass/Private/IrradianceCaching.cpp b/Engine/Plugins/Experimental/GPULightmass/Source/GPULightmass/Private/IrradianceCaching.cpp
--- a/Engine/Plugins/Experimental/GPULightmass/Source/GPULightmass/Private/IrradianceCaching.cpp
+++ b/Engine/Plugins/Experimental/GPULightmass/Source/GPULightmass/Private/IrradianceCaching.cpp
@@ -15,34 +15,37 @@ FIrradianceCache::FIrradianceCache(int32 Quality, float Spacing, float CornerRej
 {
 	uint32 IrradianceCacheTotalBytes = 0;
 
+	FRHICommandListImmediate& RHICmdList = FRHICommandListExecutor::GetImmediateCommandList();
+
 	{
-		TResourceArray<FIrradianceCacheRecord> AABBMinMax;
-		AABBMinMax.AddZeroed(IrradianceCacheMaxSize);
-		FRHIResourceCreateInfo CreateInfo(TEXT("FIrradianceCache"), &AABBMinMax);
+		FRHIResourceCreateInfo CreateInfo(TEXT("FIrradianceCache"));
 		IrradianceCacheRecords = RHICreateStructuredBuffer(sizeof(FIrradianceCacheRecord), sizeof(FIrradianceCacheRecord) * IrradianceCacheMaxSize, BUF_UnorderedAccess | BUF_ShaderResource, CreateInfo);
 		IrradianceCacheRecordsUAV = RHICreateUnorderedAccessView(IrradianceCacheRecords, false, false);
 
 		IrradianceCacheTotalBytes += sizeof(FIrradianceCacheRecord) * IrradianceCacheMaxSize;
+
+		RHICmdList.ClearUAVUint(IrradianceCacheRecordsUAV, FUintVector4(0, 0, 0, 0));
 	}
 
 	int32 HashTableSize = IrradianceCacheMaxSize * 4;
 
 	{
 		{
-			TResourceArray<uint32> EmptyHashTable;
-			EmptyHashTable.AddZeroed(HashTableSize);
-			HashTable.Initialize(TEXT("ICHashTable"), sizeof(uint32), HashTableSize, PF_R32_UINT, BUF_UnorderedAccess | BUF_ShaderResource, &EmptyHashTable);
-			HashToIndex.Initialize(TEXT("ICHashToIndex"), sizeof(uint32), HashTableSize, PF_R32_UINT, BUF_UnorderedAccess | BUF_ShaderResource, &EmptyHashTable);
-			IndexToHash.Initialize(TEXT("ICIndexToHash"), sizeof(uint32), HashTableSize, PF_R32_UINT, BUF_UnorderedAccess | BUF_ShaderResource, &EmptyHashTable);
+			HashTable.Initialize(TEXT("ICHashTable"), sizeof(uint32), HashTableSize, PF_R32_UINT, BUF_UnorderedAccess | BUF_ShaderResource);
+			HashToIndex.Initialize(TEXT("ICHashToIndex"), sizeof(uint32), HashTableSize, PF_R32_UINT, BUF_UnorderedAccess | BUF_ShaderResource);
+			IndexToHash.Initialize(TEXT("ICIndexToHash"), sizeof(uint32), HashTableSize, PF_R32_UINT, BUF_UnorderedAccess | BUF_ShaderResource);
 
 			IrradianceCacheTotalBytes += HashTable.NumBytes;
 			IrradianceCacheTotalBytes += HashToIndex.NumBytes;
 			IrradianceCacheTotalBytes += IndexToHash.NumBytes;
+
+			RHICmdList.ClearUAVUint(HashTable.UAV, FUintVector4(0, 0, 0, 0));
+			RHICmdList.ClearUAVUint(HashToIndex.UAV, FUintVector4(0, 0, 0, 0));
+			RHICmdList.ClearUAVUint(IndexToHash.UAV, FUintVector4(0, 0, 0, 0));
 		}
 		{
-			TResourceArray<uint32> ZeroRecordAllocator;
-			ZeroRecordAllocator.AddZeroed(1);
-			RecordAllocator.Initialize(TEXT("ICAllocator"), sizeof(uint32), 1, PF_R32_UINT, BUF_UnorderedAccess | BUF_ShaderResource, &ZeroRecordAllocator);
+			RecordAllocator.Initialize(TEXT("ICAllocator"), sizeof(uint32), 1, PF_R32_UINT, BUF_UnorderedAccess | BUF_ShaderResource);
+			RHICmdList.ClearUAVUint(RecordAllocator.UAV, FUintVector4(0, 0, 0, 0));
 
 			IrradianceCacheTotalBytes += 4;
 		}
diff --git a/Engine/Plugins/Experimental/GPULightmass/Source/GPULightmass/Private/LightmapRenderer.cpp b/Engine/Plugins/Experimental/GPULightmass/Source/GPULightmass/Private/LightmapRenderer.cpp
--- a/Engine/Plugins/Experimental/GPULightmass/Source/GPULightmass/Private/LightmapRenderer.cpp
+++ b/Engine/Plugins/Experimental/GPULightmass/Source/GPULightmass/Private/LightmapRenderer.cpp
@@ -851,7 +851,7 @@ void FSceneRenderState::SetupRayTracingScene()
 
 		SCOPED_DRAW_EVENTF(RHICmdList, GPULightmassUpdateRayTracingScene, TEXT("GPULightmass UpdateRayTracingScene %d Instances"), StaticMeshInstanceRenderStates.Elements.Num());
 
-		TArray<FRayTracingGeometryInstance> RayTracingGeometryInstances;
+		TArray<FRayTracingGeometryInstance, SceneRenderingAllocator>& RayTracingGeometryInstances = View.RayTracingGeometryInstances;
 		RayTracingGeometryInstances.Append(CachedRayTracingScene->RayTracingGeometryInstances);
 
 		int32 LandscapeStartOffset = RayTracingGeometryInstances.Num();
@@ -1812,6 +1812,13 @@ void FLightmapRenderer::Finalize(FRDGBuilder& GraphBuilder)
 
 				if (SampleIndex % AAvsGIMultiplier == 0)
 				{
+					TArray<int32>& PendingGIRenderPassIndices = *GraphBuilder.AllocObject<TArray<int32>>();
+
+					for (auto& Tile : PendingGITileRequests)
+					{
+						PendingGIRenderPassIndices.Add(Tile.RenderState->RetrieveTileState(Tile.VirtualCoordinates).RenderPassIndex);
+					}
+
 					for (int ScratchLayerIndex = 0; ScratchLayerIndex < 3; ScratchLayerIndex++)
 					{
 						RDG_GPU_MASK_SCOPE(GraphBuilder, FRHIGPUMask::All());
@@ -1852,10 +1859,12 @@ void FLightmapRenderer::Finalize(FRDGBuilder& GraphBuilder)
 								RDG_EVENT_NAME("LightmapGBuffer"),
 								PassParameters,
 								ERDGPassFlags::Raster,
-								[this, ReferenceView = Scene->ReferenceView, &PendingGITileRequests, GPUIndex, AAvsGIMultiplier](FRHICommandList& RHICmdList)
+								[this, ReferenceView = Scene->ReferenceView, &PendingGITileRequests, &PendingGIRenderPassIndices, GPUIndex, AAvsGIMultiplier](FRHICommandList& RHICmdList)
 							{
-								for (const FLightmapTileRequest& Tile : PendingGITileRequests)
+								for (int32 Index = 0; Index < PendingGITileRequests.Num(); Index++)
 								{
+									const FLightmapTileRequest& Tile = PendingGITileRequests[Index];
+
 									if (Tile.RenderState->IsTileGIConverged(Tile.VirtualCoordinates, Scene->Settings->GISamples)) continue;
 									uint32 AssignedGPUIndex = (Tile.RenderState->DistributionPrefixSum + Tile.RenderState->RetrieveTileStateIndex(Tile.VirtualCoordinates)) % GNumExplicitGPUsForRendering;
 									if (AssignedGPUIndex != GPUIndex) continue;
@@ -1884,7 +1893,7 @@ void FLightmapRenderer::Finalize(FRDGBuilder& GraphBuilder)
 											View = ReferenceView.Get(),
 											MeshBatches,
 											VirtualTexturePhysicalTileCoordinateScaleAndBias,
-											RenderPassIndex = Tile.RenderState->RetrieveTileState(Tile.VirtualCoordinates).RenderPassIndex / AAvsGIMultiplier,
+											RenderPassIndex = PendingGIRenderPassIndices[Index] / AAvsGIMultiplier,
 											ScratchTilePoolOffset = ScratchTilePoolGPU->GetPositionFromLinearAddress(Tile.TileAddressInScratch) * GPreviewLightmapPhysicalTileSize
 										](FDynamicPassMeshDrawListContext* DynamicMeshPassContext)
 									{
