GPULM: Port new SkyDome implementation for LightmapPathTracer

diff --git a/Engine/Plugins/Experimental/GPULightmass/Source/GPULightmass/Private/LightmapRayTracing.h b/Engine/Plugins/Experimental/GPULightmass/Source/GPULightmass/Private/LightmapRayTracing.h
--- a/Engine/Plugins/Experimental/GPULightmass/Source/GPULightmass/Private/LightmapRayTracing.h
+++ b/Engine/Plugins/Experimental/GPULightmass/Source/GPULightmass/Private/LightmapRayTracing.h
@@ -51,6 +51,7 @@ class FLightmapPathTracingRGS : public FGlobalShader
 	{
 		OutEnvironment.SetDefine(TEXT("GPreviewLightmapPhysicalTileSize"), GPreviewLightmapPhysicalTileSize);
 		OutEnvironment.SetDefine(TEXT("LIGHTMAP_PATH_TRACING_MAIN_RG"), 1);
+		OutEnvironment.SetDefine(TEXT("USE_NEW_SKYDOME"), 1);
 		OutEnvironment.CompilerFlags.Add(CFLAG_ForceDXC);
 	}
 
@@ -71,7 +72,13 @@ class FLightmapPathTracingRGS : public FGlobalShader
 		SHADER_PARAMETER_SRV(StructuredBuffer<FGPUTileDescription>, BatchedTiles)
 		SHADER_PARAMETER_RDG_BUFFER_SRV(StructuredBuffer<FPathTracingLight>, SceneLights)
 		SHADER_PARAMETER(uint32, SceneLightCount)
-		SHADER_PARAMETER_STRUCT_REF(FSkyLightData, SkyLight)
+		// Skylight
+		SHADER_PARAMETER_RDG_TEXTURE(Texture2D, SkylightTexture)
+		SHADER_PARAMETER_RDG_TEXTURE(Texture2D, SkylightPdf)
+		SHADER_PARAMETER_SAMPLER(SamplerState, SkylightTextureSampler)
+		SHADER_PARAMETER(float, SkylightInvResolution)
+		SHADER_PARAMETER(int32, SkylightMipCount)
+
 		SHADER_PARAMETER_STRUCT_REF(FViewUniformShaderParameters, ViewUniformBuffer)
 		SHADER_PARAMETER_STRUCT_REF(FIrradianceCachingParameters, IrradianceCachingParameters)
 		SHADER_PARAMETER_TEXTURE(Texture3D, IESTexture)
@@ -96,6 +103,7 @@ class FVolumetricLightmapPathTracingRGS : public FGlobalShader
 	static void ModifyCompilationEnvironment(const FGlobalShaderPermutationParameters& Parameters, FShaderCompilerEnvironment& OutEnvironment)
 	{
 		OutEnvironment.SetDefine(TEXT("GPreviewLightmapPhysicalTileSize"), GPreviewLightmapPhysicalTileSize);
+		OutEnvironment.SetDefine(TEXT("USE_NEW_SKYDOME"), 1);
 		OutEnvironment.CompilerFlags.Add(CFLAG_ForceDXC);
 	}
 
@@ -124,7 +132,13 @@ class FVolumetricLightmapPathTracingRGS : public FGlobalShader
 		SHADER_PARAMETER_UAV(RWTexture3D<float4>, OutSHCoefficients1B)
 		SHADER_PARAMETER_RDG_BUFFER_SRV(StructuredBuffer<FPathTracingLight>, SceneLights)
 		SHADER_PARAMETER(uint32, SceneLightCount)
-		SHADER_PARAMETER_STRUCT_REF(FSkyLightData, SkyLight)
+		// Skylight
+		SHADER_PARAMETER_RDG_TEXTURE(Texture2D, SkylightTexture)
+		SHADER_PARAMETER_RDG_TEXTURE(Texture2D, SkylightPdf)
+		SHADER_PARAMETER_SAMPLER(SamplerState, SkylightTextureSampler)
+		SHADER_PARAMETER(float, SkylightInvResolution)
+		SHADER_PARAMETER(int32, SkylightMipCount)
+
 		SHADER_PARAMETER_STRUCT_REF(FViewUniformShaderParameters, ViewUniformBuffer)
 		SHADER_PARAMETER_STRUCT_REF(FIrradianceCachingParameters, IrradianceCachingParameters)
 		SHADER_PARAMETER_TEXTURE(Texture3D, IESTexture)
diff --git a/Engine/Plugins/Experimental/GPULightmass/Source/GPULightmass/Private/LightmapRenderer.cpp b/Engine/Plugins/Experimental/GPULightmass/Source/GPULightmass/Private/LightmapRenderer.cpp
--- a/Engine/Plugins/Experimental/GPULightmass/Source/GPULightmass/Private/LightmapRenderer.cpp
+++ b/Engine/Plugins/Experimental/GPULightmass/Source/GPULightmass/Private/LightmapRenderer.cpp
@@ -8,6 +8,7 @@
 #include "Scene/StaticMesh.h"
 #include "LightmapGBuffer.h"
 #include "LightmapRayTracing.h"
+#include "PathTracingLightParameters.inl"
 #include "ClearQuad.h"
 #include "RayTracing/RayTracingMaterialHitShaders.h"
 #include "Async/ParallelFor.h"
@@ -165,201 +166,6 @@ struct FGPUBatchedTileRequests
 	TResourceArray<FGPUTileDescription> BatchedTilesDesc;
 };
 
-#if RHI_RAYTRACING
-void SetupPathTracingLightParameters(const GPULightmass::FLightSceneRenderState& LightScene, FRDGBuilder& GraphBuilder, FRDGBufferSRV** OutLightBuffer, unsigned* OutLightCount)
-{
-	const uint32 MaxLightCount = RAY_TRACING_LIGHT_COUNT_MAXIMUM;
-	FPathTracingLight Lights[RAY_TRACING_LIGHT_COUNT_MAXIMUM] = {};
-
-	unsigned LightCount = 0;
-
-	if (LightScene.SkyLight.IsSet())
-	{
-		FPathTracingLight& DestLight = Lights[LightCount++];
-		DestLight.Color = FVector(LightScene.SkyLight->Color);
-		DestLight.Flags = PATHTRACER_FLAG_TRANSMISSION_MASK;
-		DestLight.Flags |= PATHTRACER_FLAG_LIGHTING_CHANNEL_MASK;
-		bool SkyLightIsStationary = LightScene.SkyLight->bStationary;
-		DestLight.Flags |= SkyLightIsStationary ? PATHTRACER_FLAG_STATIONARY_MASK : 0;
-		DestLight.Flags |= PATHTRACING_LIGHT_SKY;
-	}
-
-	for (auto Light : LightScene.DirectionalLights.Elements)
-	{
-		if (LightCount < MaxLightCount)
-		{
-			FPathTracingLight& DestLight = Lights[LightCount++];
-
-			DestLight.Normal = -Light.Direction;
-			DestLight.Color = FVector(Light.Color);
-			DestLight.Dimensions = FVector(
-				FMath::Sin(0.5f * FMath::DegreesToRadians(Light.LightSourceAngle)),
-				FMath::Sin(0.5f * FMath::DegreesToRadians(Light.LightSourceSoftAngle)),
-				0.0f);
-			DestLight.Attenuation = 1.0;
-			DestLight.IESTextureSlice = -1;
-
-			DestLight.Flags = PATHTRACER_FLAG_TRANSMISSION_MASK;
-			DestLight.Flags |= PATHTRACER_FLAG_LIGHTING_CHANNEL_MASK;
-			DestLight.Flags |= Light.bStationary ? PATHTRACER_FLAG_STATIONARY_MASK : 0;
-			DestLight.Flags |= PATHTRACING_LIGHT_DIRECTIONAL;
-		}
-	}
-
-	for (auto Light : LightScene.PointLights.Elements)
-	{
-		if (LightCount < MaxLightCount)
-		{
-			FPathTracingLight& DestLight = Lights[LightCount++];
-
-			DestLight.Position = Light.Position;
-			DestLight.Color = FVector(Light.Color);
-			DestLight.Normal = Light.Direction;
-			DestLight.dPdu = FVector::CrossProduct(Light.Tangent, Light.Direction);
-			DestLight.dPdv = Light.Tangent;
-
-			DestLight.Dimensions = FVector(Light.SourceRadius, Light.SourceSoftRadius, Light.SourceLength);
-			DestLight.Attenuation = 1.0f / Light.AttenuationRadius;
-			DestLight.FalloffExponent = Light.FalloffExponent;
-			DestLight.IESTextureSlice = -1;
-
-			DestLight.Flags = PATHTRACER_FLAG_TRANSMISSION_MASK;
-			DestLight.Flags |= PATHTRACER_FLAG_LIGHTING_CHANNEL_MASK;
-			DestLight.Flags |= Light.IsInverseSquared ? 0 : PATHTRACER_FLAG_NON_INVERSE_SQUARE_FALLOFF_MASK;
-			DestLight.Flags |= Light.bStationary ? PATHTRACER_FLAG_STATIONARY_MASK : 0;
-			DestLight.Flags |= PATHTRACING_LIGHT_POINT;
-		}
-	}
-
-	for (auto Light : LightScene.SpotLights.Elements)
-	{
-		if (LightCount < MaxLightCount)
-		{
-			FPathTracingLight& DestLight = Lights[LightCount++];
-
-			DestLight.Position = Light.Position;
-			DestLight.Normal = Light.Direction;
-			DestLight.dPdu = FVector::CrossProduct(Light.Tangent, Light.Direction);
-			DestLight.dPdv = Light.Tangent;
-			DestLight.Color = FVector(Light.Color);
-			DestLight.Dimensions = FVector(Light.SourceRadius, Light.SourceSoftRadius, Light.SourceLength);
-			DestLight.Shaping = Light.SpotAngles;
-			DestLight.Attenuation = 1.0f / Light.AttenuationRadius;
-			DestLight.FalloffExponent = Light.FalloffExponent;
-			DestLight.IESTextureSlice = -1;
-
-			DestLight.Flags = PATHTRACER_FLAG_TRANSMISSION_MASK;
-			DestLight.Flags |= PATHTRACER_FLAG_LIGHTING_CHANNEL_MASK;
-			DestLight.Flags |= Light.IsInverseSquared ? 0 : PATHTRACER_FLAG_NON_INVERSE_SQUARE_FALLOFF_MASK;
-			DestLight.Flags |= Light.bStationary ? PATHTRACER_FLAG_STATIONARY_MASK : 0;
-			DestLight.Flags |= PATHTRACING_LIGHT_SPOT;
-		}
-	}
-
-	for (auto Light : LightScene.RectLights.Elements)
-	{
-		if (LightCount < MaxLightCount)
-		{
-			FPathTracingLight& DestLight = Lights[LightCount++];
-
-			DestLight.Position = Light.Position;
-			DestLight.Normal = Light.Direction;
-			DestLight.dPdu = FVector::CrossProduct(Light.Tangent, -Light.Direction);
-			DestLight.dPdv = Light.Tangent;
-
-			FLinearColor LightColor = Light.Color;
-			LightColor /= 0.5f * Light.SourceWidth * Light.SourceHeight;
-			DestLight.Color = FVector(LightColor);
-
-			DestLight.Dimensions = FVector(Light.SourceWidth, Light.SourceHeight, 0.0f);
-			DestLight.Attenuation = 1.0f / Light.AttenuationRadius;
-			DestLight.Shaping = FVector2D(FMath::Cos(FMath::DegreesToRadians(Light.BarnDoorAngle)), Light.BarnDoorLength);
-
-			DestLight.IESTextureSlice = -1;
-
-			DestLight.Flags = PATHTRACER_FLAG_TRANSMISSION_MASK;
-			DestLight.Flags |= PATHTRACER_FLAG_LIGHTING_CHANNEL_MASK;
-			DestLight.Flags |= Light.bStationary ? PATHTRACER_FLAG_STATIONARY_MASK : 0;
-			DestLight.Flags |= PATHTRACING_LIGHT_RECT;
-		}
-	}
-
-	{
-		// Upload the buffer of lights to the GPU
-		size_t DataSize = sizeof(FPathTracingLight) * FMath::Max(LightCount, 1u);
-		*OutLightBuffer = GraphBuilder.CreateSRV(FRDGBufferSRVDesc(CreateStructuredBuffer(GraphBuilder, TEXT("PathTracingLightsBuffer"), sizeof(FPathTracingLight), FMath::Max(LightCount, 1u), Lights, DataSize)));
-		*OutLightCount = LightCount;
-	}
-}
-
-FSkyLightData SetupSkyLightParameters(const GPULightmass::FLightSceneRenderState& LightScene)
-{
-	FSkyLightData SkyLightData;
-	// Check if parameters should be set based on if the sky light's texture has been processed and if its mip tree has been built yet
-	if (LightScene.SkyLight.IsSet())
-	{
-		check(LightScene.SkyLight->ProcessedTexture);
-		check(LightScene.SkyLight->ImportanceSamplingData->bIsValid);
-
-		SkyLightData.SamplesPerPixel = 1;
-		SkyLightData.SamplingStopLevel = 0;
-		SkyLightData.MaxRayDistance = 1.0e7;
-		SkyLightData.MaxNormalBias = 0.1f;
-		SkyLightData.MaxShadowThickness = 1.0e3;
-
-		SkyLightData.Color = FVector(LightScene.SkyLight->Color);
-		SkyLightData.Texture = LightScene.SkyLight->ProcessedTexture;
-		SkyLightData.TextureDimensions = FIntVector(LightScene.SkyLight->TextureDimensions.X, LightScene.SkyLight->TextureDimensions.Y, 1);
-		SkyLightData.TextureSampler = LightScene.SkyLight->ProcessedTextureSampler;
-		SkyLightData.MipDimensions = LightScene.SkyLight->ImportanceSamplingData->MipDimensions;
-
-		SkyLightData.MipTreePosX = LightScene.SkyLight->ImportanceSamplingData->MipTreePosX.SRV;
-		SkyLightData.MipTreeNegX = LightScene.SkyLight->ImportanceSamplingData->MipTreeNegX.SRV;
-		SkyLightData.MipTreePosY = LightScene.SkyLight->ImportanceSamplingData->MipTreePosY.SRV;
-		SkyLightData.MipTreeNegY = LightScene.SkyLight->ImportanceSamplingData->MipTreeNegY.SRV;
-		SkyLightData.MipTreePosZ = LightScene.SkyLight->ImportanceSamplingData->MipTreePosZ.SRV;
-		SkyLightData.MipTreeNegZ = LightScene.SkyLight->ImportanceSamplingData->MipTreeNegZ.SRV;
-
-		SkyLightData.MipTreePdfPosX = LightScene.SkyLight->ImportanceSamplingData->MipTreePdfPosX.SRV;
-		SkyLightData.MipTreePdfNegX = LightScene.SkyLight->ImportanceSamplingData->MipTreePdfNegX.SRV;
-		SkyLightData.MipTreePdfPosY = LightScene.SkyLight->ImportanceSamplingData->MipTreePdfPosY.SRV;
-		SkyLightData.MipTreePdfNegY = LightScene.SkyLight->ImportanceSamplingData->MipTreePdfNegY.SRV;
-		SkyLightData.MipTreePdfPosZ = LightScene.SkyLight->ImportanceSamplingData->MipTreePdfPosZ.SRV;
-		SkyLightData.MipTreePdfNegZ = LightScene.SkyLight->ImportanceSamplingData->MipTreePdfNegZ.SRV;
-		SkyLightData.SolidAnglePdf = LightScene.SkyLight->ImportanceSamplingData->SolidAnglePdf.SRV;
-	}
-	else
-	{
-		SkyLightData.SamplesPerPixel = -1;
-		SkyLightData.SamplingStopLevel = 0;
-		SkyLightData.MaxRayDistance = 0.0f;
-		SkyLightData.MaxNormalBias = 0.0f;
-		SkyLightData.MaxShadowThickness = 0.0f;
-
-		SkyLightData.Color = FVector(0.0);
-		SkyLightData.Texture = GBlackTextureCube->TextureRHI;
-		SkyLightData.TextureSampler = TStaticSamplerState<SF_Bilinear, AM_Clamp, AM_Clamp, AM_Clamp>::GetRHI();
-		SkyLightData.MipDimensions = FIntVector(0);
-
-		SkyLightData.MipTreePosX = GBlackTextureWithSRV->ShaderResourceViewRHI;
-		SkyLightData.MipTreeNegX = GBlackTextureWithSRV->ShaderResourceViewRHI;
-		SkyLightData.MipTreePosY = GBlackTextureWithSRV->ShaderResourceViewRHI;
-		SkyLightData.MipTreeNegY = GBlackTextureWithSRV->ShaderResourceViewRHI;
-		SkyLightData.MipTreePosZ = GBlackTextureWithSRV->ShaderResourceViewRHI;
-		SkyLightData.MipTreeNegZ = GBlackTextureWithSRV->ShaderResourceViewRHI;
-
-		SkyLightData.MipTreePdfPosX = GBlackTextureWithSRV->ShaderResourceViewRHI;
-		SkyLightData.MipTreePdfNegX = GBlackTextureWithSRV->ShaderResourceViewRHI;
-		SkyLightData.MipTreePdfPosY = GBlackTextureWithSRV->ShaderResourceViewRHI;
-		SkyLightData.MipTreePdfNegY = GBlackTextureWithSRV->ShaderResourceViewRHI;
-		SkyLightData.MipTreePdfPosZ = GBlackTextureWithSRV->ShaderResourceViewRHI;
-		SkyLightData.MipTreePdfNegZ = GBlackTextureWithSRV->ShaderResourceViewRHI;
-		SkyLightData.SolidAnglePdf = GBlackTextureWithSRV->ShaderResourceViewRHI;
-	}
-
-	return SkyLightData;
-}
-#endif
 namespace GPULightmass
 {
 
@@ -1954,6 +1760,7 @@ void FLightmapRenderer::Finalize(FRDGBuilder& GraphBuilder)
 						if (GPUBatchedTileRequests.BatchedTilesDesc.Num() > 0)
 						{
 							FRHIResourceCreateInfo CreateInfo(TEXT("BatchedTilesBuffer"));
+							CreateInfo.GPUMask = FRHIGPUMask::FromIndex(GPUIndex);
 							CreateInfo.ResourceArray = &GPUBatchedTileRequests.BatchedTilesDesc;
 
 							GPUBatchedTileRequests.BatchedTilesBuffer = RHICreateStructuredBuffer(sizeof(FGPUTileDescription), GPUBatchedTileRequests.BatchedTilesDesc.GetResourceDataSize(), BUF_Static | BUF_ShaderResource, CreateInfo);
@@ -2015,14 +1822,7 @@ void FLightmapRenderer::Finalize(FRDGBuilder& GraphBuilder)
 									PassParameters->ViewUniformBuffer = Scene->ReferenceView->ViewUniformBuffer;
 									PassParameters->IrradianceCachingParameters = Scene->IrradianceCache->IrradianceCachingParametersUniformBuffer;
 
-									{
-										SetupPathTracingLightParameters(Scene->LightSceneRenderState, GraphBuilder, &PassParameters->SceneLights, &PassParameters->SceneLightCount);
-									}
-
-									{
-										SkyLightDataUniformBuffer = CreateUniformBufferImmediate(SetupSkyLightParameters(Scene->LightSceneRenderState), EUniformBufferUsage::UniformBuffer_SingleFrame);
-										PassParameters->SkyLight = SkyLightDataUniformBuffer;
-									}
+									SetupPathTracingLightParameters(Scene->LightSceneRenderState, GraphBuilder, PassParameters);
 
 									// TODO: find a way to share IES atlas with path tracer ...
 									PassParameters->IESTexture = GWhiteTexture->TextureRHI;
diff --git a/Engine/Plugins/Experimental/GPULightmass/Source/GPULightmass/Private/PathTracingLightParameters.inl b/Engine/Plugins/Experimental/GPULightmass/Source/GPULightmass/Private/PathTracingLightParameters.inl
new file mode 100644
--- /dev/null
+++ b/Engine/Plugins/Experimental/GPULightmass/Source/GPULightmass/Private/PathTracingLightParameters.inl
@@ -0,0 +1,144 @@
+template<typename PassParameterType>
+void SetupPathTracingLightParameters(
+	const GPULightmass::FLightSceneRenderState& LightScene,
+	FRDGBuilder& GraphBuilder,
+	PassParameterType* PassParameters)
+{
+	const uint32 MaxLightCount = RAY_TRACING_LIGHT_COUNT_MAXIMUM;
+	FPathTracingLight Lights[RAY_TRACING_LIGHT_COUNT_MAXIMUM] = {};
+
+	unsigned LightCount = 0;
+
+	if (LightScene.SkyLight.IsSet())
+	{
+		FPathTracingLight& DestLight = Lights[LightCount++];
+		DestLight.Color = FVector(LightScene.SkyLight->Color);
+		DestLight.Flags = PATHTRACER_FLAG_TRANSMISSION_MASK;
+		DestLight.Flags |= PATHTRACER_FLAG_LIGHTING_CHANNEL_MASK;
+		bool SkyLightIsStationary = LightScene.SkyLight->bStationary;
+		DestLight.Flags |= SkyLightIsStationary ? PATHTRACER_FLAG_STATIONARY_MASK : 0;
+		DestLight.Flags |= PATHTRACING_LIGHT_SKY;
+
+		PassParameters->SkylightTexture = GraphBuilder.RegisterExternalTexture(LightScene.SkyLight->PathTracingSkylightTexture, TEXT("PathTracer.Skylight"));
+		PassParameters->SkylightTextureSampler = TStaticSamplerState<SF_Bilinear, AM_Clamp, AM_Clamp, AM_Clamp>::GetRHI();
+		PassParameters->SkylightPdf = GraphBuilder.RegisterExternalTexture(LightScene.SkyLight->PathTracingSkylightPdf, TEXT("PathTracer.SkylightPdf"));
+		PassParameters->SkylightInvResolution = LightScene.SkyLight->SkylightInvResolution;
+		PassParameters->SkylightMipCount = LightScene.SkyLight->SkylightMipCount;
+	}
+	else
+	{
+		PassParameters->SkylightTexture = GraphBuilder.RegisterExternalTexture(GSystemTextures.BlackDummy);
+		PassParameters->SkylightTextureSampler = TStaticSamplerState<SF_Bilinear, AM_Clamp, AM_Clamp, AM_Clamp>::GetRHI();
+		PassParameters->SkylightPdf = GraphBuilder.RegisterExternalTexture(GSystemTextures.BlackDummy);
+		PassParameters->SkylightInvResolution = 0;
+		PassParameters->SkylightMipCount = 0;
+	}
+
+
+	for (auto Light : LightScene.DirectionalLights.Elements)
+	{
+		if (LightCount < MaxLightCount)
+		{
+			FPathTracingLight& DestLight = Lights[LightCount++];
+
+			DestLight.Normal = -Light.Direction;
+			DestLight.Color = FVector(Light.Color);
+			DestLight.Dimensions = FVector(
+				FMath::Sin(0.5f * FMath::DegreesToRadians(Light.LightSourceAngle)),
+				FMath::Sin(0.5f * FMath::DegreesToRadians(Light.LightSourceSoftAngle)),
+				0.0f);
+			DestLight.Attenuation = 1.0;
+			DestLight.IESTextureSlice = -1;
+
+			DestLight.Flags = PATHTRACER_FLAG_TRANSMISSION_MASK;
+			DestLight.Flags |= PATHTRACER_FLAG_LIGHTING_CHANNEL_MASK;
+			DestLight.Flags |= Light.bStationary ? PATHTRACER_FLAG_STATIONARY_MASK : 0;
+			DestLight.Flags |= PATHTRACING_LIGHT_DIRECTIONAL;
+		}
+	}
+
+	for (auto Light : LightScene.PointLights.Elements)
+	{
+		if (LightCount < MaxLightCount)
+		{
+			FPathTracingLight& DestLight = Lights[LightCount++];
+
+			DestLight.Position = Light.Position;
+			DestLight.Color = FVector(Light.Color);
+			DestLight.Normal = Light.Direction;
+			DestLight.dPdu = FVector::CrossProduct(Light.Tangent, Light.Direction);
+			DestLight.dPdv = Light.Tangent;
+
+			DestLight.Dimensions = FVector(Light.SourceRadius, Light.SourceSoftRadius, Light.SourceLength);
+			DestLight.Attenuation = 1.0f / Light.AttenuationRadius;
+			DestLight.FalloffExponent = Light.FalloffExponent;
+			DestLight.IESTextureSlice = -1;
+
+			DestLight.Flags = PATHTRACER_FLAG_TRANSMISSION_MASK;
+			DestLight.Flags |= PATHTRACER_FLAG_LIGHTING_CHANNEL_MASK;
+			DestLight.Flags |= Light.IsInverseSquared ? 0 : PATHTRACER_FLAG_NON_INVERSE_SQUARE_FALLOFF_MASK;
+			DestLight.Flags |= Light.bStationary ? PATHTRACER_FLAG_STATIONARY_MASK : 0;
+			DestLight.Flags |= PATHTRACING_LIGHT_POINT;
+		}
+	}
+
+	for (auto Light : LightScene.SpotLights.Elements)
+	{
+		if (LightCount < MaxLightCount)
+		{
+			FPathTracingLight& DestLight = Lights[LightCount++];
+
+			DestLight.Position = Light.Position;
+			DestLight.Normal = Light.Direction;
+			DestLight.dPdu = FVector::CrossProduct(Light.Tangent, Light.Direction);
+			DestLight.dPdv = Light.Tangent;
+			DestLight.Color = FVector(Light.Color);
+			DestLight.Dimensions = FVector(Light.SourceRadius, Light.SourceSoftRadius, Light.SourceLength);
+			DestLight.Shaping = Light.SpotAngles;
+			DestLight.Attenuation = 1.0f / Light.AttenuationRadius;
+			DestLight.FalloffExponent = Light.FalloffExponent;
+			DestLight.IESTextureSlice = -1;
+
+			DestLight.Flags = PATHTRACER_FLAG_TRANSMISSION_MASK;
+			DestLight.Flags |= PATHTRACER_FLAG_LIGHTING_CHANNEL_MASK;
+			DestLight.Flags |= Light.IsInverseSquared ? 0 : PATHTRACER_FLAG_NON_INVERSE_SQUARE_FALLOFF_MASK;
+			DestLight.Flags |= Light.bStationary ? PATHTRACER_FLAG_STATIONARY_MASK : 0;
+			DestLight.Flags |= PATHTRACING_LIGHT_SPOT;
+		}
+	}
+
+	for (auto Light : LightScene.RectLights.Elements)
+	{
+		if (LightCount < MaxLightCount)
+		{
+			FPathTracingLight& DestLight = Lights[LightCount++];
+
+			DestLight.Position = Light.Position;
+			DestLight.Normal = Light.Direction;
+			DestLight.dPdu = FVector::CrossProduct(Light.Tangent, -Light.Direction);
+			DestLight.dPdv = Light.Tangent;
+
+			FLinearColor LightColor = Light.Color;
+			LightColor /= 0.5f * Light.SourceWidth * Light.SourceHeight;
+			DestLight.Color = FVector(LightColor);
+
+			DestLight.Dimensions = FVector(Light.SourceWidth, Light.SourceHeight, 0.0f);
+			DestLight.Attenuation = 1.0f / Light.AttenuationRadius;
+			DestLight.Shaping = FVector2D(FMath::Cos(FMath::DegreesToRadians(Light.BarnDoorAngle)), Light.BarnDoorLength);
+
+			DestLight.IESTextureSlice = -1;
+
+			DestLight.Flags = PATHTRACER_FLAG_TRANSMISSION_MASK;
+			DestLight.Flags |= PATHTRACER_FLAG_LIGHTING_CHANNEL_MASK;
+			DestLight.Flags |= Light.bStationary ? PATHTRACER_FLAG_STATIONARY_MASK : 0;
+			DestLight.Flags |= PATHTRACING_LIGHT_RECT;
+		}
+	}
+
+	{
+		// Upload the buffer of lights to the GPU
+		size_t DataSize = sizeof(FPathTracingLight) * FMath::Max(LightCount, 1u);
+		PassParameters->SceneLights = GraphBuilder.CreateSRV(FRDGBufferSRVDesc(CreateStructuredBuffer(GraphBuilder, TEXT("PathTracingLightsBuffer"), sizeof(FPathTracingLight), FMath::Max(LightCount, 1u), Lights, DataSize)));
+		PassParameters->SceneLightCount = LightCount;
+	}
+}
diff --git a/Engine/Plugins/Experimental/GPULightmass/Source/GPULightmass/Private/Scene/Lights.cpp b/Engine/Plugins/Experimental/GPULightmass/Source/GPULightmass/Private/Scene/Lights.cpp
--- a/Engine/Plugins/Experimental/GPULightmass/Source/GPULightmass/Private/Scene/Lights.cpp
+++ b/Engine/Plugins/Experimental/GPULightmass/Source/GPULightmass/Private/Scene/Lights.cpp
@@ -1,6 +1,22 @@
 // Copyright Epic Games, Inc. All Rights Reserved.
 
 #include "Lights.h"
+#include "RenderGraphBuilder.h"
+#include "ReflectionEnvironment.h"
+
+RENDERER_API void PrepareSkyTexture_Internal(
+	FRDGBuilder& GraphBuilder,
+	FReflectionUniformParameters& Parameters,
+	uint32 Size,
+	FLinearColor SkyColor,
+	bool UseMISCompensation,
+
+	// Out
+	FRDGTextureRef& SkylightTexture,
+	FRDGTextureRef& SkylightPdf,
+	float& SkylightInvResolution,
+	int32& SkylightMipCount
+);
 
 namespace GPULightmass
 {
@@ -269,4 +285,43 @@ FLightShaderParameters FRectLightRenderState::GetLightShaderParameters() const
 	return LightParameters;
 }
 
+void FSkyLightRenderState::PrepareSkyTexture(FRHICommandListImmediate& RHICmdList)
+{
+	FRDGBuilder GraphBuilder(RHICmdList);
+
+	FRDGTextureRef SkylightTexture;
+	FRDGTextureRef SkylightPdf;
+
+	FReflectionUniformParameters Parameters;
+	Parameters.SkyLightCubemap = ProcessedTexture;
+	Parameters.SkyLightCubemapSampler = ProcessedTextureSampler;
+	Parameters.SkyLightBlendDestinationCubemap = ProcessedTexture;
+	Parameters.SkyLightBlendDestinationCubemapSampler = ProcessedTextureSampler;
+	Parameters.SkyLightParameters = FVector4(1, 1, 0, 0);
+
+	FLinearColor SkyColor = Color;
+	// since we are resampled into an octahedral layout, we multiply the cubemap resolution by 2 to get roughly the same number of texels
+	uint32 Size = FMath::RoundUpToPowerOfTwo(2 * TextureDimensions.X);
+
+	const bool UseMISCompensation = true;
+
+	PrepareSkyTexture_Internal(
+		GraphBuilder,
+		Parameters,
+		Size,
+		SkyColor,
+		UseMISCompensation,
+		// Out
+		SkylightTexture,
+		SkylightPdf,
+		SkylightInvResolution,
+		SkylightMipCount
+	);
+
+	GraphBuilder.QueueTextureExtraction(SkylightTexture, &PathTracingSkylightTexture);
+	GraphBuilder.QueueTextureExtraction(SkylightPdf, &PathTracingSkylightPdf);
+
+	GraphBuilder.Execute();
+}
+
 }
diff --git a/Engine/Plugins/Experimental/GPULightmass/Source/GPULightmass/Private/Scene/Lights.h b/Engine/Plugins/Experimental/GPULightmass/Source/GPULightmass/Private/Scene/Lights.h
--- a/Engine/Plugins/Experimental/GPULightmass/Source/GPULightmass/Private/Scene/Lights.h
+++ b/Engine/Plugins/Experimental/GPULightmass/Source/GPULightmass/Private/Scene/Lights.h
@@ -280,7 +280,13 @@ struct FSkyLightRenderState
 	FIntPoint TextureDimensions;
 	FSHVectorRGB3 IrradianceEnvironmentMap;
 	FRWBufferStructured SkyIrradianceEnvironmentMap;
-	const FSkyLightImportanceSamplingData* ImportanceSamplingData;
+
+	// New sky dome
+	void PrepareSkyTexture(FRHICommandListImmediate& RHICmdList);
+	float SkylightInvResolution;
+	int32 SkylightMipCount;
+	TRefCountPtr<IPooledRenderTarget> PathTracingSkylightTexture;
+	TRefCountPtr<IPooledRenderTarget> PathTracingSkylightPdf;
 };
 
 struct FSkyLightBuildInfo
diff --git a/Engine/Plugins/Experimental/GPULightmass/Source/GPULightmass/Private/Scene/Scene.cpp b/Engine/Plugins/Experimental/GPULightmass/Source/GPULightmass/Private/Scene/Scene.cpp
--- a/Engine/Plugins/Experimental/GPULightmass/Source/GPULightmass/Private/Scene/Scene.cpp
+++ b/Engine/Plugins/Experimental/GPULightmass/Source/GPULightmass/Private/Scene/Scene.cpp
@@ -601,9 +601,7 @@ void FScene::AddLight(USkyLightComponent* SkyLight)
 	NewSkyLightRenderState.Color = SkyLight->GetLightColor() * SkyLight->Intensity;
 	NewSkyLightRenderState.TextureDimensions = FIntPoint(SkyLight->GetProcessedSkyTexture()->GetSizeX(), SkyLight->GetProcessedSkyTexture()->GetSizeY());
 	NewSkyLightRenderState.IrradianceEnvironmentMap = SkyLight->GetIrradianceEnvironmentMap();
-#if RHI_RAYTRACING
-	NewSkyLightRenderState.ImportanceSamplingData = SkyLight->GetImportanceSamplingData();
-#endif
+
 	ENQUEUE_RENDER_COMMAND(AddLightRenderState)(
 		[&RenderState = RenderState, NewSkyLightRenderState = MoveTemp(NewSkyLightRenderState), ProcessedSkyTexture = SkyLight->GetProcessedSkyTexture()](FRHICommandListImmediate& RHICmdList) mutable
 	{
@@ -613,6 +611,8 @@ void FScene::AddLight(USkyLightComponent* SkyLight)
 
 		NewSkyLightRenderState.SkyIrradianceEnvironmentMap.Initialize(TEXT("SkyIrradianceEnvironmentMap"), sizeof(FVector4), 7, 0);
 
+		NewSkyLightRenderState.PrepareSkyTexture(RHICmdList);
+
 		// Set the captured environment map data
 		void* DataPtr = RHICmdList.LockBuffer(NewSkyLightRenderState.SkyIrradianceEnvironmentMap.Buffer, 0, NewSkyLightRenderState.SkyIrradianceEnvironmentMap.NumBytes, RLM_WriteOnly);
 		SetupSkyIrradianceEnvironmentMapConstantsFromSkyIrradiance((FVector4*)DataPtr, NewSkyLightRenderState.IrradianceEnvironmentMap);
diff --git a/Engine/Plugins/Experimental/GPULightmass/Source/GPULightmass/Private/VolumetricLightmap.cpp b/Engine/Plugins/Experimental/GPULightmass/Source/GPULightmass/Private/VolumetricLightmap.cpp
--- a/Engine/Plugins/Experimental/GPULightmass/Source/GPULightmass/Private/VolumetricLightmap.cpp
+++ b/Engine/Plugins/Experimental/GPULightmass/Source/GPULightmass/Private/VolumetricLightmap.cpp
@@ -5,6 +5,7 @@
 #include "ScenePrivate.h"
 #include "Containers/DynamicRHIResourceArray.h"
 #include "LightmapRayTracing.h"
+#include "PathTracingLightParameters.inl"
 #include "GPULightmassModule.h"
 #include "RHIGPUReadback.h"
 #include "LevelEditorViewport.h"
@@ -557,9 +558,7 @@ void FVolumetricLightmapRenderer::BackgroundTick()
 					PassParameters->ViewUniformBuffer = Scene->ReferenceView->ViewUniformBuffer;
 					PassParameters->IrradianceCachingParameters = Scene->IrradianceCache->IrradianceCachingParametersUniformBuffer;
 
-					SetupPathTracingLightParameters(Scene->LightSceneRenderState, GraphBuilder, &PassParameters->SceneLights, &PassParameters->SceneLightCount);
-
-					PassParameters->SkyLight = CreateUniformBufferImmediate(SetupSkyLightParameters(Scene->LightSceneRenderState), EUniformBufferUsage::UniformBuffer_SingleFrame);
+					SetupPathTracingLightParameters(Scene->LightSceneRenderState, GraphBuilder, PassParameters);
 
 					// TODO: find a way to share IES atlas with path tracer ...
 					PassParameters->IESTexture = GWhiteTexture->TextureRHI;
diff --git a/Engine/Source/Runtime/Renderer/Private/PathTracing.cpp b/Engine/Source/Runtime/Renderer/Private/PathTracing.cpp
--- a/Engine/Source/Runtime/Renderer/Private/PathTracing.cpp
+++ b/Engine/Source/Runtime/Renderer/Private/PathTracing.cpp
@@ -440,6 +440,76 @@
 };
 IMPLEMENT_SHADER_TYPE(, FPathTracingIESAtlasCS, TEXT("/Engine/Private/PathTracing/PathTracingIESAtlas.usf"), TEXT("PathTracingIESAtlasCS"), SF_Compute);
 
+RENDERER_API void PrepareSkyTexture_Internal(
+	FRDGBuilder& GraphBuilder,
+	FReflectionUniformParameters& Parameters,
+	uint32 Size,
+	FLinearColor SkyColor,
+	bool UseMISCompensation,
+
+	// Out
+	FRDGTextureRef& SkylightTexture,
+	FRDGTextureRef& SkylightPdf,
+	float& SkylightInvResolution,
+	int32& SkylightMipCount
+)
+{
+	FRDGTextureDesc SkylightTextureDesc = FRDGTextureDesc::Create2D(
+		FIntPoint(Size, Size),
+		PF_A32B32G32R32F, // half precision might be ok?
+		FClearValueBinding::None,
+		TexCreate_ShaderResource | TexCreate_UAV);
+	SkylightTexture = GraphBuilder.CreateTexture(SkylightTextureDesc, TEXT("PathTracer.Skylight"), ERDGTextureFlags::None);
+	FRDGTextureDesc SkylightPdfDesc = FRDGTextureDesc::Create2D(
+		FIntPoint(Size, Size),
+		PF_R32_FLOAT, // half precision might be ok?
+		FClearValueBinding::None,
+		TexCreate_ShaderResource | TexCreate_UAV,
+		FMath::CeilLogTwo(Size) + 1);
+	SkylightPdf = GraphBuilder.CreateTexture(SkylightPdfDesc, TEXT("PathTracer.SkylightPdf"), ERDGTextureFlags::None);
+
+	SkylightInvResolution = 1.0f / Size;
+	SkylightMipCount = SkylightPdfDesc.NumMips;
+
+	// run a simple compute shader to sample the cubemap and prep the top level of the mipmap hierarchy
+	{
+		TShaderMapRef<FPathTracingSkylightPrepareCS> ComputeShader(GetGlobalShaderMap(GMaxRHIFeatureLevel));
+		FPathTracingSkylightPrepareCS::FParameters* PassParameters = GraphBuilder.AllocParameters<FPathTracingSkylightPrepareCS::FParameters>();
+		PassParameters->SkyColor = FVector(SkyColor.R, SkyColor.G, SkyColor.B);
+		PassParameters->SkyLightCubemap0 = Parameters.SkyLightCubemap;
+		PassParameters->SkyLightCubemap1 = Parameters.SkyLightBlendDestinationCubemap;
+		PassParameters->SkyLightCubemapSampler0 = Parameters.SkyLightCubemapSampler;
+		PassParameters->SkyLightCubemapSampler1 = Parameters.SkyLightBlendDestinationCubemapSampler;
+		PassParameters->SkylightBlendFactor = Parameters.SkyLightParameters.W;
+		PassParameters->SkylightInvResolution = SkylightInvResolution;
+		PassParameters->SkylightTextureOutput = GraphBuilder.CreateUAV(FRDGTextureUAVDesc(SkylightTexture, 0));
+		PassParameters->SkylightTexturePdf = GraphBuilder.CreateUAV(FRDGTextureUAVDesc(SkylightPdf, 0));
+		FComputeShaderUtils::AddPass(
+			GraphBuilder,
+			RDG_EVENT_NAME("SkylightPrepare"),
+			ComputeShader,
+			PassParameters,
+			FComputeShaderUtils::GetGroupCount(FIntPoint(Size, Size), FComputeShaderUtils::kGolden2DGroupSize));
+	}
+	FGenerateMips::ExecuteCompute(GraphBuilder, SkylightPdf, TStaticSamplerState<SF_Bilinear, AM_Clamp, AM_Clamp, AM_Clamp>::GetRHI());
+
+	if (UseMISCompensation)
+	{
+		TShaderMapRef<FPathTracingSkylightMISCompensationCS> ComputeShader(GetGlobalShaderMap(GMaxRHIFeatureLevel));
+		FPathTracingSkylightMISCompensationCS::FParameters* PassParameters = GraphBuilder.AllocParameters<FPathTracingSkylightMISCompensationCS::FParameters>();
+		PassParameters->SkylightTexturePdfAverage = GraphBuilder.CreateSRV(FRDGTextureSRVDesc::CreateForMipLevel(SkylightPdf, SkylightMipCount - 1));
+		PassParameters->SkylightTextureOutput = GraphBuilder.CreateUAV(FRDGTextureUAVDesc(SkylightTexture, 0));
+		PassParameters->SkylightTexturePdf = GraphBuilder.CreateUAV(FRDGTextureUAVDesc(SkylightPdf, 0));
+		FComputeShaderUtils::AddPass(
+			GraphBuilder,
+			RDG_EVENT_NAME("SkylightMISCompensation"),
+			ComputeShader,
+			PassParameters,
+			FComputeShaderUtils::GetGroupCount(FIntPoint(Size, Size), FComputeShaderUtils::kGolden2DGroupSize));
+		FGenerateMips::ExecuteCompute(GraphBuilder, SkylightPdf, TStaticSamplerState<SF_Bilinear, AM_Clamp, AM_Clamp, AM_Clamp>::GetRHI());
+	}
+}
+
 bool PrepareSkyTexture(FRDGBuilder& GraphBuilder, FScene* Scene, const FViewInfo& View, bool UseMISCompensation, FPathTracingRG::FParameters* PathTracingParameters)
 {
 	PathTracingParameters->SkylightTextureSampler = TStaticSamplerState<SF_Bilinear, AM_Clamp, AM_Clamp, AM_Clamp>::GetRHI();
@@ -483,60 +553,18 @@
 	// since we are resampled into an octahedral layout, we multiply the cubemap resolution by 2 to get roughly the same number of texels
 	uint32 Size = FMath::RoundUpToPowerOfTwo(2 * Scene->SkyLight->CaptureCubeMapResolution);
 
-	FRDGTextureDesc SkylightTextureDesc = FRDGTextureDesc::Create2D(
-		FIntPoint(Size, Size),
-		PF_A32B32G32R32F, // half precision might be ok?
-		FClearValueBinding::None,
-		TexCreate_ShaderResource | TexCreate_UAV);
-	PathTracingParameters->SkylightTexture = GraphBuilder.CreateTexture(SkylightTextureDesc, TEXT("PathTracer.Skylight"), ERDGTextureFlags::None);
-	FRDGTextureDesc SkylightPdfDesc = FRDGTextureDesc::Create2D(
-		FIntPoint(Size, Size),
-		PF_R32_FLOAT, // half precision might be ok?
-		FClearValueBinding::None,
-		TexCreate_ShaderResource | TexCreate_UAV,
-		FMath::CeilLogTwo(Size) + 1);
-	PathTracingParameters->SkylightPdf = GraphBuilder.CreateTexture(SkylightPdfDesc, TEXT("PathTracer.SkylightPdf"), ERDGTextureFlags::None);
+	PrepareSkyTexture_Internal(
+		GraphBuilder,
+		Parameters,
+		Size,
+		SkyColor,
+		UseMISCompensation,
+		// Out
+		PathTracingParameters->SkylightTexture,
+		PathTracingParameters->SkylightPdf,
+		PathTracingParameters->SkylightInvResolution,
+		PathTracingParameters->SkylightMipCount
+	);
 
-	PathTracingParameters->SkylightInvResolution = 1.0f / Size;
-	PathTracingParameters->SkylightMipCount = SkylightPdfDesc.NumMips;
-
-	// run a simple compute shader to sample the cubemap and prep the top level of the mipmap hierarchy
-	{
-		TShaderMapRef<FPathTracingSkylightPrepareCS> ComputeShader(View.ShaderMap);
-		FPathTracingSkylightPrepareCS::FParameters* PassParameters = GraphBuilder.AllocParameters<FPathTracingSkylightPrepareCS::FParameters>();
-		PassParameters->SkyColor = FVector(SkyColor.R, SkyColor.G, SkyColor.B);
-		PassParameters->SkyLightCubemap0 = Parameters.SkyLightCubemap;
-		PassParameters->SkyLightCubemap1 = Parameters.SkyLightBlendDestinationCubemap;
-		PassParameters->SkyLightCubemapSampler0 = Parameters.SkyLightCubemapSampler;
-		PassParameters->SkyLightCubemapSampler1 = Parameters.SkyLightBlendDestinationCubemapSampler;
-		PassParameters->SkylightBlendFactor = Parameters.SkyLightParameters.W;
-		PassParameters->SkylightInvResolution = PathTracingParameters->SkylightInvResolution;
-		PassParameters->SkylightTextureOutput = GraphBuilder.CreateUAV(FRDGTextureUAVDesc(PathTracingParameters->SkylightTexture, 0));
-		PassParameters->SkylightTexturePdf = GraphBuilder.CreateUAV(FRDGTextureUAVDesc(PathTracingParameters->SkylightPdf, 0));
-		FComputeShaderUtils::AddPass(
-			GraphBuilder,
-			RDG_EVENT_NAME("SkylightPrepare"),
-			ComputeShader,
-			PassParameters,
-			FComputeShaderUtils::GetGroupCount(FIntPoint(Size, Size), FComputeShaderUtils::kGolden2DGroupSize));
-	}
-	FGenerateMips::ExecuteCompute(GraphBuilder, PathTracingParameters->SkylightPdf, TStaticSamplerState<SF_Bilinear, AM_Clamp, AM_Clamp, AM_Clamp>::GetRHI());
-
-	if (UseMISCompensation)
-	{
-		TShaderMapRef<FPathTracingSkylightMISCompensationCS> ComputeShader(View.ShaderMap);
-		FPathTracingSkylightMISCompensationCS::FParameters* PassParameters = GraphBuilder.AllocParameters<FPathTracingSkylightMISCompensationCS::FParameters>();
-		PassParameters->SkylightTexturePdfAverage = GraphBuilder.CreateSRV(FRDGTextureSRVDesc::CreateForMipLevel(PathTracingParameters->SkylightPdf, PathTracingParameters->SkylightMipCount - 1));
-		PassParameters->SkylightTextureOutput = GraphBuilder.CreateUAV(FRDGTextureUAVDesc(PathTracingParameters->SkylightTexture, 0));
-		PassParameters->SkylightTexturePdf = GraphBuilder.CreateUAV(FRDGTextureUAVDesc(PathTracingParameters->SkylightPdf, 0));
-		FComputeShaderUtils::AddPass(
-			GraphBuilder,
-			RDG_EVENT_NAME("SkylightMISCompensation"),
-			ComputeShader,
-			PassParameters,
-			FComputeShaderUtils::GetGroupCount(FIntPoint(Size, Size), FComputeShaderUtils::kGolden2DGroupSize));
-		FGenerateMips::ExecuteCompute(GraphBuilder, PathTracingParameters->SkylightPdf, TStaticSamplerState<SF_Bilinear, AM_Clamp, AM_Clamp, AM_Clamp>::GetRHI());
-	}
-
 	// hang onto these for next time (if caching is enabled)
 	if (IsSkylightCachingEnabled)
