Fix unity build

diff --git a/Engine/Plugins/Experimental/GPULightmass/Source/GPULightmass/Private/PathTracingLightParameters.inl b/Engine/Plugins/Experimental/GPULightmass/Source/GPULightmass/Private/PathTracingLightParameters.inl
--- a/Engine/Plugins/Experimental/GPULightmass/Source/GPULightmass/Private/PathTracingLightParameters.inl
+++ b/Engine/Plugins/Experimental/GPULightmass/Source/GPULightmass/Private/PathTracingLightParameters.inl
@@ -1,3 +1,7 @@
+// Copyright Epic Games, Inc. All Rights Reserved.
+
+#pragma once
+
 template<typename PassParameterType>
 void SetupPathTracingLightParameters(
 	const GPULightmass::FLightSceneRenderState& LightScene,
