GPULM: Implement GPU instace culling changes

diff --git a/Engine/Plugins/Experimental/GPULightmass/Source/GPULightmass/Private/LightmapRenderer.cpp b/Engine/Plugins/Experimental/GPULightmass/Source/GPULightmass/Private/LightmapRenderer.cpp
--- a/Engine/Plugins/Experimental/GPULightmass/Source/GPULightmass/Private/LightmapRenderer.cpp
+++ b/Engine/Plugins/Experimental/GPULightmass/Source/GPULightmass/Private/LightmapRenderer.cpp
@@ -24,9 +24,8 @@
 #include "EngineModule.h"
 #include "PostProcess/PostProcessing.h"
 #include "RayTracingGeometryManager.h"
-#include "RayTracingDefinitions.h"
-#include "PathTracingDefinitions.h"
 #include "ScreenPass.h"
+#include "RayTracingDynamicGeometryCollection.h"
 
 class FCopyConvergedLightmapTilesCS : public FGlobalShader
 {
@@ -234,6 +233,19 @@ void FLightmapRenderer::AddRequest(FLightmapTileRequest TileRequest)
 	PendingTileRequests.AddUnique(TileRequest);
 }
 
+int32 FSceneRenderState::GetPrimitiveIdForGPUScene(const FGeometryInstanceRenderStateRef& GeometryInstanceRef) const
+{
+	int32 PrimitiveId = GeometryInstanceRef.GetElementIdChecked();
+	if (StaticMeshInstanceRenderStates.Contains(GeometryInstanceRef)) return PrimitiveId;
+	PrimitiveId += StaticMeshInstanceRenderStates.Elements.Num();
+	if (InstanceGroupRenderStates.Contains(GeometryInstanceRef)) return PrimitiveId;
+	PrimitiveId += InstanceGroupRenderStates.Elements.Num();
+	if (LandscapeRenderStates.Contains(GeometryInstanceRef)) return PrimitiveId;
+
+	checkf(false, TEXT("The referenced geometry isn't in any of the geometry arrays"));
+	return INDEX_NONE;
+}
+
 void FCachedRayTracingSceneData::SetupViewUniformBufferFromSceneRenderState(FSceneRenderState& Scene)
 {
 	TResourceArray<FPrimitiveSceneShaderData> PrimitiveSceneData;
@@ -242,22 +254,40 @@ void FCachedRayTracingSceneData::SetupViewUniformBufferFromSceneRenderState(FSce
 
 	PrimitiveSceneData.AddZeroed(Scene.StaticMeshInstanceRenderStates.Elements.Num());
 	InstanceSceneData.AddZeroed(Scene.StaticMeshInstanceRenderStates.Elements.Num());
+	InstanceDataOriginalOffsets.AddZeroed(Scene.StaticMeshInstanceRenderStates.Elements.Num());
 
 	TArray<int32> LightmapSceneDataStartOffsets;
 	LightmapSceneDataStartOffsets.AddZeroed(Scene.StaticMeshInstanceRenderStates.Elements.Num());
+	LightmapSceneDataStartOffsets.AddZeroed(Scene.InstanceGroupRenderStates.Elements.Num());
+	LightmapSceneDataStartOffsets.AddZeroed(Scene.LandscapeRenderStates.Elements.Num());
 
 	{
 		TRACE_CPUPROFILER_EVENT_SCOPE(ComputePrefixSum);
 
 		int32 ConservativeLightmapEntriesNum = 0;
+		int32 PrimitiveId = 0;
 
 		for (int32 InstanceIndex = 0; InstanceIndex < Scene.StaticMeshInstanceRenderStates.Elements.Num(); InstanceIndex++)
 		{
 			FStaticMeshInstanceRenderState& Instance = Scene.StaticMeshInstanceRenderStates.Elements[InstanceIndex];
-			LightmapSceneDataStartOffsets[InstanceIndex] = ConservativeLightmapEntriesNum;
+			LightmapSceneDataStartOffsets[PrimitiveId++] = ConservativeLightmapEntriesNum;
 			ConservativeLightmapEntriesNum += Instance.LODLightmapRenderStates.Num();
 		}
 
+		for (int32 InstanceGroupIndex = 0; InstanceGroupIndex < Scene.InstanceGroupRenderStates.Elements.Num(); InstanceGroupIndex++)
+		{
+			FInstanceGroupRenderState& InstanceGroup = Scene.InstanceGroupRenderStates.Elements[InstanceGroupIndex];
+			LightmapSceneDataStartOffsets[PrimitiveId++] = ConservativeLightmapEntriesNum;
+			ConservativeLightmapEntriesNum += InstanceGroup.LODLightmapRenderStates.Num();
+		}
+
+		for (int32 LandscapeIndex = 0; LandscapeIndex < Scene.LandscapeRenderStates.Elements.Num(); LandscapeIndex++)
+		{
+			FLandscapeRenderState& Landscape = Scene.LandscapeRenderStates.Elements[LandscapeIndex];
+			LightmapSceneDataStartOffsets[PrimitiveId++] = ConservativeLightmapEntriesNum;
+			ConservativeLightmapEntriesNum += Landscape.LODLightmapRenderStates.Num();
+		}
+
 		for (int32 InstanceGroupIndex = 0; InstanceGroupIndex < Scene.InstanceGroupRenderStates.Elements.Num(); InstanceGroupIndex++)
 		{
 			FInstanceGroupRenderState& InstanceGroup = Scene.InstanceGroupRenderStates.Elements[InstanceGroupIndex];
@@ -278,7 +308,8 @@ void FCachedRayTracingSceneData::SetupViewUniformBufferFromSceneRenderState(FSce
 				&PrimitiveSceneData,
 				&LightmapSceneData,
 				&LightmapSceneDataStartOffsets,
-				&InstanceSceneData
+				&InstanceSceneData,
+				&InstanceDataOriginalOffsets = InstanceDataOriginalOffsets
 			](int32 InstanceIndex = 0)
 			{
 				FStaticMeshInstanceRenderState& Instance = Scene.StaticMeshInstanceRenderStates.Elements[InstanceIndex];
@@ -287,6 +318,7 @@ void FCachedRayTracingSceneData::SetupViewUniformBufferFromSceneRenderState(FSce
 				PrimitiveUniformShaderParameters.LightmapDataIndex = LightmapSceneDataStartOffsets[InstanceIndex];
 				PrimitiveUniformShaderParameters.LightmapUVIndex = 0; // TODO: LightmapUVIndex
 				PrimitiveUniformShaderParameters.InstanceDataOffset = InstanceIndex;
+				InstanceDataOriginalOffsets[InstanceIndex] = InstanceIndex;
 				PrimitiveUniformShaderParameters.NumInstanceDataEntries = 1;
 				PrimitiveSceneData[InstanceIndex] = FPrimitiveSceneShaderData(PrimitiveUniformShaderParameters);
 				InstanceSceneData[InstanceIndex] = FInstanceSceneShaderData(
@@ -324,6 +356,8 @@ void FCachedRayTracingSceneData::SetupViewUniformBufferFromSceneRenderState(FSce
 				}
 			});
 
+		int32 PrimitiveId = Scene.StaticMeshInstanceRenderStates.Elements.Num();
+
 		for (int32 InstanceGroupIndex = 0; InstanceGroupIndex < Scene.InstanceGroupRenderStates.Elements.Num(); InstanceGroupIndex++)
 		{
 			FInstanceGroupRenderState& InstanceGroup = Scene.InstanceGroupRenderStates.Elements[InstanceGroupIndex];
@@ -344,7 +378,7 @@ void FCachedRayTracingSceneData::SetupViewUniformBufferFromSceneRenderState(FSce
 				false,
 				false,
 				0b111,
-				LightmapSceneDataStartOffsets[Scene.StaticMeshInstanceRenderStates.Elements.Num() + InstanceGroupIndex],
+				LightmapSceneDataStartOffsets[PrimitiveId],
 				INDEX_NONE,
 				INDEX_NONE,
 				/* bOutputVelocity = */ false,
@@ -355,6 +389,27 @@ void FCachedRayTracingSceneData::SetupViewUniformBufferFromSceneRenderState(FSce
 				/* bCastShadow = */ true
 			);
 
+			InstanceDataOriginalOffsets.Add(InstanceSceneData.Num());
+
+			for (int32 LODIndex = 0; LODIndex < InstanceGroup.LODLightmapRenderStates.Num(); LODIndex++)
+			{
+				FPrecomputedLightingUniformParameters LightmapParams;
+				GetDefaultPrecomputedLightingParameters(LightmapParams);
+
+				if (InstanceGroup.LODLightmapRenderStates[LODIndex].IsValid())
+				{
+					LightmapParams.LightmapVTPackedPageTableUniform[0] = InstanceGroup.LODLightmapRenderStates[LODIndex]->LightmapVTPackedPageTableUniform[0];
+					for (uint32 LayerIndex = 0u; LayerIndex < 5u; ++LayerIndex)
+					{
+						LightmapParams.LightmapVTPackedUniform[LayerIndex] = InstanceGroup.LODLightmapRenderStates[LODIndex]->LightmapVTPackedUniform[LayerIndex];
+					}
+
+					LightmapParams.LightMapCoordinateScaleBias = InstanceGroup.LODLightmapRenderStates[LODIndex]->LightmapCoordinateScaleBias;
+				}
+
+				LightmapSceneData[LightmapSceneDataStartOffsets[PrimitiveId] + LODIndex] = FLightmapSceneShaderData(LightmapParams);
+			}
+
 			for (int32 InstanceIdx = 0; InstanceIdx < NumInstancesThisGroup; InstanceIdx++)
 			{
 				FMatrix Transform;
@@ -363,7 +418,7 @@ void FCachedRayTracingSceneData::SetupViewUniformBufferFromSceneRenderState(FSce
 				FMatrix InstanceTransform = Transform * InstanceGroup.LocalToWorld;
 
 				FPrimitiveInstance Instance;
-				Instance.PrimitiveId = PrimitiveSceneData.Num();
+				Instance.PrimitiveId = PrimitiveId;
 				Instance.InstanceToLocal = Transform;
 				// Filled in during GPU Scene update...
 				Instance.LocalToWorld = InstanceTransform;
@@ -374,8 +429,79 @@ void FCachedRayTracingSceneData::SetupViewUniformBufferFromSceneRenderState(FSce
 				InstanceSceneData.Add(FInstanceSceneShaderData(Instance));
 			}
 
-			InstanceGroup.UniformBuffer = TUniformBufferRef<FPrimitiveUniformShaderParameters>::CreateUniformBufferImmediate(PrimitiveUniformShaderParameters, UniformBuffer_MultiFrame);
 			PrimitiveSceneData.Add(FPrimitiveSceneShaderData(PrimitiveUniformShaderParameters));
+
+			PrimitiveId++;
+		}
+
+		for (int32 LandscapeIndex = 0; LandscapeIndex < Scene.LandscapeRenderStates.Elements.Num(); LandscapeIndex++)
+		{
+			FLandscapeRenderState& Landscape = Scene.LandscapeRenderStates.Elements[LandscapeIndex];
+			
+			FPrimitiveUniformShaderParameters PrimitiveUniformShaderParameters = GetPrimitiveUniformShaderParameters(
+				Landscape.LocalToWorld,
+				Landscape.LocalToWorld,
+				Landscape.ActorPosition,
+				Landscape.WorldBounds,
+				Landscape.LocalBounds,
+				Landscape.LocalBounds,
+				false,
+				false,
+				false,
+				false,
+				false,
+				false,
+				0b111,
+				LightmapSceneDataStartOffsets[PrimitiveId],
+				INDEX_NONE,
+				INDEX_NONE,
+				/* bOutputVelocity = */ false,
+				nullptr,
+				/* bCastContactShadow = */ true,
+				InstanceSceneData.Num(),
+				1,
+				/* bCastShadow = */ true
+			);
+
+			InstanceDataOriginalOffsets.Add(InstanceSceneData.Num());
+
+			for (int32 LODIndex = 0; LODIndex < Landscape.LODLightmapRenderStates.Num(); LODIndex++)
+			{
+				FPrecomputedLightingUniformParameters LightmapParams;
+				GetDefaultPrecomputedLightingParameters(LightmapParams);
+
+				if (Landscape.LODLightmapRenderStates[LODIndex].IsValid())
+				{
+					LightmapParams.LightmapVTPackedPageTableUniform[0] = Landscape.LODLightmapRenderStates[LODIndex]->LightmapVTPackedPageTableUniform[0];
+					for (uint32 LayerIndex = 0u; LayerIndex < 5u; ++LayerIndex)
+					{
+						LightmapParams.LightmapVTPackedUniform[LayerIndex] = Landscape.LODLightmapRenderStates[LODIndex]->LightmapVTPackedUniform[LayerIndex];
+					}
+
+					LightmapParams.LightMapCoordinateScaleBias = Landscape.LODLightmapRenderStates[LODIndex]->LightmapCoordinateScaleBias;
+				}
+
+				LightmapSceneData[LightmapSceneDataStartOffsets[PrimitiveId] + LODIndex] = FLightmapSceneShaderData(LightmapParams);
+			}
+
+			InstanceSceneData.Add(FInstanceSceneShaderData(
+				GetInstanceUniformShaderParameters(
+					PrimitiveUniformShaderParameters.LocalToWorld,
+					PrimitiveUniformShaderParameters.PreviousLocalToWorld,
+					(PrimitiveUniformShaderParameters.LocalObjectBoundsMin + PrimitiveUniformShaderParameters.LocalObjectBoundsMax) * 0.5f,
+					PrimitiveUniformShaderParameters.LocalObjectBoundsMax - PrimitiveUniformShaderParameters.LocalObjectBoundsMin,
+					PrimitiveUniformShaderParameters.NonUniformScale, PrimitiveUniformShaderParameters.InvNonUniformScaleAndDeterminantSign,
+					FVector4(ForceInitToZero),
+					FNaniteInfo(),
+					PrimitiveId,
+					0xFFFFFFFFu,
+					0.0f,
+					true
+				)));
+
+			PrimitiveSceneData.Add(FPrimitiveSceneShaderData(PrimitiveUniformShaderParameters));
+
+			PrimitiveId++;
 		}
 	}
 
@@ -432,6 +558,25 @@ void FCachedRayTracingSceneData::SetupViewUniformBufferFromSceneRenderState(FSce
 			InstanceSceneDataSOAStride = InstanceSceneData.Num();
 		}
 
+		{
+			TRACE_CPUPROFILER_EVENT_SCOPE(InstanceIdsIdentityBuffer);
+
+			TResourceArray<uint32> InstanceIdsIdentity;
+			for (int32 Index = 0; Index < InstanceSceneData.Num(); Index++)
+			{
+				InstanceIdsIdentity.Add(Index);
+			}
+
+			FRHIResourceCreateInfo CreateInfo(TEXT("InstanceIdsIdentityBuffer"), &InstanceIdsIdentity);
+			if (InstanceIdsIdentity.GetResourceDataSize() == 0)
+			{
+				InstanceIdsIdentity.Add(0);
+			}
+
+			InstanceIdsIdentityBufferRHI = RHICreateStructuredBuffer(sizeof(uint32), InstanceIdsIdentity.GetResourceDataSize(), BUF_Static | BUF_ShaderResource, CreateInfo);
+			InstanceIdsIdentityBufferSRV = RHICreateShaderResourceView(InstanceIdsIdentityBufferRHI);
+		}
+
 		FViewUniformShaderParameters ViewUniformBufferParameters;
 		CachedViewUniformBuffer = TUniformBufferRef<FViewUniformShaderParameters>::CreateUniformBufferImmediate(ViewUniformBufferParameters, UniformBuffer_MultiFrame, EUniformBufferValidation::None);
 	}
@@ -629,6 +774,7 @@ void FSceneRenderState::SetupRayTracingScene()
 				*View.CachedViewUniformShaderParameters);
 
 			View.CachedViewUniformShaderParameters->InstanceDataSOAStride = CachedRayTracingScene->InstanceSceneDataSOAStride;
+			View.CachedViewUniformShaderParameters->InstanceIdsBuffer = CachedRayTracingScene->InstanceIdsIdentityBufferSRV;
 
 			if (LightSceneRenderState.SkyLight.IsSet())
 			{
@@ -682,8 +828,10 @@ void FSceneRenderState::SetupRayTracingScene()
 
 			int32 NumLandscapeInstances = 0;
 
-			for (FLandscapeRenderState& Landscape : LandscapeRenderStates.Elements)
+			for (int32 LandscapeIndex = 0; LandscapeIndex < LandscapeRenderStates.Elements.Num(); LandscapeIndex++)
 			{
+				FLandscapeRenderState& Landscape = LandscapeRenderStates.Elements[LandscapeIndex];
+
 				for (int32 SubY = 0; SubY < Landscape.NumSubsections; SubY++)
 				{
 					for (int32 SubX = 0; SubX < Landscape.NumSubsections; SubX++)
@@ -694,6 +842,75 @@ void FSceneRenderState::SetupRayTracingScene()
 						int32 InstanceIndex = LandscapeStartOffset + NumLandscapeInstances;
 						NumLandscapeInstances++;
 
+						if (!Landscape.SectionRayTracingStates[SubSectionIdx].IsValid())
+						{
+							FRayTracingGeometryInitializer GeometryInitializer;
+							GeometryInitializer.IndexBuffer = Landscape.SharedBuffers->ZeroOffsetIndexBuffers[0]->IndexBufferRHI;
+							GeometryInitializer.TotalPrimitiveCount = NumPrimitives;
+							GeometryInitializer.GeometryType = RTGT_Triangles;
+							GeometryInitializer.bFastBuild = false;
+							GeometryInitializer.bAllowUpdate = false;
+
+							FRayTracingGeometrySegment Segment;
+							Segment.VertexBuffer = nullptr;
+							Segment.VertexBufferStride = sizeof(FVector);
+							Segment.VertexBufferElementType = VET_Float3;
+							Segment.MaxVertices = NumPrimitives * 3;
+							Segment.NumPrimitives = NumPrimitives;
+							GeometryInitializer.Segments.Add(Segment);
+
+							Landscape.SectionRayTracingStates[SubSectionIdx] = MakeUnique<FLandscapeRenderState::FLandscapeSectionRayTracingState>();
+							Landscape.SectionRayTracingStates[SubSectionIdx]->Geometry.SetInitializer(GeometryInitializer);
+							Landscape.SectionRayTracingStates[SubSectionIdx]->Geometry.InitResource();
+
+							FRayTracingDynamicGeometryCollection DynamicGeometryCollection;
+
+							TArray<FMeshBatch> MeshBatches = Landscape.GetMeshBatchesForGBufferRendering(0);
+
+							FLandscapeVertexFactoryMVFParameters UniformBufferParams;
+							UniformBufferParams.SubXY = FIntPoint(SubX, SubY);
+							Landscape.SectionRayTracingStates[SubSectionIdx]->UniformBuffer = FLandscapeVertexFactoryMVFUniformBufferRef::CreateUniformBufferImmediate(UniformBufferParams, UniformBuffer_MultiFrame);
+
+							FLandscapeBatchElementParams& BatchElementParams = *(FLandscapeBatchElementParams*)MeshBatches[0].Elements[0].UserData;
+							BatchElementParams.LandscapeVertexFactoryMVFUniformBuffer = Landscape.SectionRayTracingStates[SubSectionIdx]->UniformBuffer;
+
+							MeshBatches[0].Elements[0].IndexBuffer = Landscape.SharedBuffers->ZeroOffsetIndexBuffers[0];
+							MeshBatches[0].Elements[0].FirstIndex = 0;
+							MeshBatches[0].Elements[0].NumPrimitives = NumPrimitives;
+							MeshBatches[0].Elements[0].MinVertexIndex = 0;
+							MeshBatches[0].Elements[0].MaxVertexIndex = 0;
+
+							MeshBatches[0].Elements[0].DynamicPrimitiveShaderDataIndex = LandscapeIndex;
+							MeshBatches[0].Elements[0].DynamicPrimitiveShaderDataIndex += StaticMeshInstanceRenderStates.Elements.Num();
+							MeshBatches[0].Elements[0].DynamicPrimitiveShaderDataIndex += InstanceGroupRenderStates.Elements.Num();
+
+							FRayTracingDynamicGeometryUpdateParams UpdateParams
+							{
+								MeshBatches,
+								false,
+								(uint32)FMath::Square(Landscape.SubsectionSizeVerts),
+								FMath::Square(Landscape.SubsectionSizeVerts) * (uint32)sizeof(FVector),
+								(uint32)FMath::Square(Landscape.SubsectionSizeVerts - 1) * 2,
+								&Landscape.SectionRayTracingStates[SubSectionIdx]->Geometry,
+								&Landscape.SectionRayTracingStates[SubSectionIdx]->RayTracingDynamicVertexBuffer,
+								false
+							};
+
+							DynamicGeometryCollection.AddDynamicMeshBatchForGeometryUpdate(
+								Landscape.ComponentUObject->GetWorld()->Scene->GetRenderScene(),
+								ReferenceView.Get(),
+								nullptr,
+								UpdateParams,
+								StaticMeshInstanceRenderStates.Elements.Num() + InstanceGroupRenderStates.Elements.Num() + LandscapeIndex
+							);
+
+							DynamicGeometryCollection.DispatchUpdates(RHICmdList);
+
+							// Landscape VF doesn't really use the vertex buffer in HitGroupSystemParameters
+							// We can release after all related RHI cmds get dispatched onto the cmd list
+							Landscape.SectionRayTracingStates[SubSectionIdx]->RayTracingDynamicVertexBuffer.Release();
+						}
+
 						FRayTracingGeometryInstance& RayTracingInstance = RayTracingGeometryInstances[InstanceIndex];
 						RayTracingInstance.GeometryRHI = Landscape.SectionRayTracingStates[SubSectionIdx]->Geometry.RayTracingGeometryRHI;
 						RayTracingInstance.Transforms.Add(FMatrix::Identity);
@@ -940,6 +1157,81 @@ bool ClampTexelPositionAndOffsetTile(FIntPoint& SrcVirtualTexelPosition, FIntPoi
 	return bLoadingOutOfBounds;
 }
 
+void FLightmapRenderer::RenderMeshBatchesIntoGBuffer(
+	FRHICommandList& RHICmdList,
+	const FSceneView* View,
+	int32 GPUScenePrimitiveId,
+	TArray<FMeshBatch>& MeshBatches, 
+	FVector4 VirtualTexturePhysicalTileCoordinateScaleAndBias,
+	int32 RenderPassIndex,
+	FIntPoint ScratchTilePoolOffset)
+{
+	TArray<FMeshBatch> MeshBatchesToDrawImmediately;
+	TArray<FMeshBatch> MeshBatchesNeedingInstanceOffsetUpdates;
+
+	for (auto& MeshBatch : MeshBatches)
+	{
+		FMeshBatchElement& Element = MeshBatch.Elements[0];
+		Element.DynamicPrimitiveShaderDataIndex = GPUScenePrimitiveId;
+		if (Element.UserIndex == INDEX_NONE)
+		{
+			MeshBatchesToDrawImmediately.Add(MeshBatch);
+		}
+		else
+		{
+			MeshBatchesNeedingInstanceOffsetUpdates.Add(MeshBatch);
+		}
+	}
+
+	DrawDynamicMeshPass(
+		*View, RHICmdList,
+		[
+			View,
+			MeshBatchesToDrawImmediately,
+			VirtualTexturePhysicalTileCoordinateScaleAndBias,
+			RenderPassIndex,
+			ScratchTilePoolOffset
+		](FDynamicPassMeshDrawListContext* DynamicMeshPassContext)
+	{
+		FLightmapGBufferMeshProcessor MeshProcessor(nullptr, View, DynamicMeshPassContext, VirtualTexturePhysicalTileCoordinateScaleAndBias, RenderPassIndex, ScratchTilePoolOffset);
+
+		for (auto& MeshBatch : MeshBatchesToDrawImmediately)
+		{
+			MeshProcessor.AddMeshBatch(MeshBatch, ~0ull, nullptr);
+		}
+	});
+
+	for (auto& MeshBatch : MeshBatchesNeedingInstanceOffsetUpdates)
+	{
+		FDynamicMeshDrawCommandStorage DynamicMeshDrawCommandStorage;
+		FMeshCommandOneFrameArray VisibleMeshDrawCommands;
+		FGraphicsMinimalPipelineStateSet GraphicsMinimalPipelineStateSet;
+		bool NeedsShaderInitialisation;
+
+		FDynamicPassMeshDrawListContext DynamicMeshPassContext(DynamicMeshDrawCommandStorage, VisibleMeshDrawCommands, GraphicsMinimalPipelineStateSet, NeedsShaderInitialisation);
+
+		FLightmapGBufferMeshProcessor MeshProcessor(nullptr, View, &DynamicMeshPassContext, VirtualTexturePhysicalTileCoordinateScaleAndBias, RenderPassIndex, ScratchTilePoolOffset);
+
+		MeshProcessor.AddMeshBatch(MeshBatch, ~0ull, nullptr);
+
+		const uint32 InstanceFactor = 1;
+		FRHIBuffer* PrimitiveIdVertexBuffer = nullptr;
+		const bool bDynamicInstancing = false;
+
+		for (FVisibleMeshDrawCommand& Cmd : VisibleMeshDrawCommands)
+		{
+			Cmd.DrawPrimitiveId = Scene->CachedRayTracingScene->InstanceDataOriginalOffsets[MeshBatch.Elements[0].DynamicPrimitiveShaderDataIndex] + MeshBatch.Elements[0].UserIndex;
+		}
+
+		TRange<int32> Range = TRange<int32>::Empty();
+		Range.SetLowerBoundValue(INDEX_NONE);
+		SortAndMergeDynamicPassMeshDrawCommands(View->GetFeatureLevel(), VisibleMeshDrawCommands, DynamicMeshDrawCommandStorage, PrimitiveIdVertexBuffer, InstanceFactor, Range);
+		SubmitMeshDrawCommands(VisibleMeshDrawCommands, GraphicsMinimalPipelineStateSet, PrimitiveIdVertexBuffer, 0, bDynamicInstancing, InstanceFactor, RHICmdList);
+	}
+
+	GPrimitiveIdVertexBufferPool.DiscardAll();
+}
+
 void FLightmapRenderer::Finalize(FRDGBuilder& GraphBuilder)
 {
 	TRACE_CPUPROFILER_EVENT_SCOPE(FLightmapRenderer::Finalize);
@@ -1686,30 +1978,14 @@ void FLightmapRenderer::Finalize(FRDGBuilder& GraphBuilder)
 
 									TArray<FMeshBatch> MeshBatches = Tile.RenderState->GeometryInstanceRef.GetMeshBatchesForGBufferRendering(Tile.VirtualCoordinates);
 
-									for (auto& MeshBatch : MeshBatches)
-									{
-										FMeshBatchElement& Element = MeshBatch.Elements[0];
-
-										Element.DynamicPrimitiveShaderDataIndex = Tile.RenderState->GeometryInstanceRef.GetElementId();
-									}
-
-									DrawDynamicMeshPass(
-										*ReferenceView, RHICmdList,
-										[
-											View = ReferenceView.Get(),
-											MeshBatches,
-											VirtualTexturePhysicalTileCoordinateScaleAndBias,
-											RenderPassIndex = PendingGIRenderPassIndices[Index] / AAvsGIMultiplier,
-											ScratchTilePoolOffset = ScratchTilePoolGPU->GetPositionFromLinearAddress(Tile.TileAddressInScratch) * GPreviewLightmapPhysicalTileSize
-										](FDynamicPassMeshDrawListContext* DynamicMeshPassContext)
-									{
-										FLightmapGBufferMeshProcessor MeshProcessor(nullptr, View, DynamicMeshPassContext, VirtualTexturePhysicalTileCoordinateScaleAndBias, RenderPassIndex, ScratchTilePoolOffset);
-
-										for (auto& MeshBatch : MeshBatches)
-										{
-											MeshProcessor.AddMeshBatch(MeshBatch, ~0ull, nullptr);
-										}
-									});
+									RenderMeshBatchesIntoGBuffer(
+										RHICmdList,
+										ReferenceView.Get(),
+										Scene->GetPrimitiveIdForGPUScene(Tile.RenderState->GeometryInstanceRef),
+										MeshBatches,
+										VirtualTexturePhysicalTileCoordinateScaleAndBias,
+										PendingGIRenderPassIndices[Index] / AAvsGIMultiplier,
+										ScratchTilePoolGPU->GetPositionFromLinearAddress(Tile.TileAddressInScratch)* GPreviewLightmapPhysicalTileSize);
 
 									GPrimitiveIdVertexBufferPool.DiscardAll();
 								}
@@ -2130,30 +2406,14 @@ void FLightmapRenderer::Finalize(FRDGBuilder& GraphBuilder)
 
 								TArray<FMeshBatch> MeshBatches = Tile.RenderState->GeometryInstanceRef.GetMeshBatchesForGBufferRendering(Tile.VirtualCoordinates);
 
-								for (auto& MeshBatch : MeshBatches)
-								{
-									FMeshBatchElement& Element = MeshBatch.Elements[0];
-
-									Element.DynamicPrimitiveShaderDataIndex = Tile.RenderState->GeometryInstanceRef.GetElementId();
-								}
-
-								DrawDynamicMeshPass(
-									*ReferenceView, RHICmdList,
-									[
-										View = ReferenceView.Get(),
-										MeshBatches,
-										VirtualTexturePhysicalTileCoordinateScaleAndBias,
-										RenderPassIndex = LightSampleIndexArray[TileIndex],
-										ScratchTilePoolOffset = ScratchTilePoolGPU->GetPositionFromLinearAddress(Tile.TileAddressInScratch) * GPreviewLightmapPhysicalTileSize
-									](FDynamicPassMeshDrawListContext* DynamicMeshPassContext)
-								{
-									FLightmapGBufferMeshProcessor MeshProcessor(nullptr, View, DynamicMeshPassContext, VirtualTexturePhysicalTileCoordinateScaleAndBias, RenderPassIndex, ScratchTilePoolOffset);
-
-									for (auto& MeshBatch : MeshBatches)
-									{
-										MeshProcessor.AddMeshBatch(MeshBatch, ~0ull, nullptr);
-									}
-								});
+								RenderMeshBatchesIntoGBuffer(
+									RHICmdList,
+									ReferenceView.Get(),
+									Scene->GetPrimitiveIdForGPUScene(Tile.RenderState->GeometryInstanceRef),
+									MeshBatches,
+									VirtualTexturePhysicalTileCoordinateScaleAndBias,
+									LightSampleIndexArray[TileIndex],
+									ScratchTilePoolGPU->GetPositionFromLinearAddress(Tile.TileAddressInScratch)* GPreviewLightmapPhysicalTileSize);
 
 								GPrimitiveIdVertexBufferPool.DiscardAll();
 							}
diff --git a/Engine/Plugins/Experimental/GPULightmass/Source/GPULightmass/Private/LightmapRenderer.h b/Engine/Plugins/Experimental/GPULightmass/Source/GPULightmass/Private/LightmapRenderer.h
--- a/Engine/Plugins/Experimental/GPULightmass/Source/GPULightmass/Private/LightmapRenderer.h
+++ b/Engine/Plugins/Experimental/GPULightmass/Source/GPULightmass/Private/LightmapRenderer.h
@@ -110,6 +110,15 @@ public:
 
 	void AddRequest(FLightmapTileRequest TileRequest);
 
+	void RenderMeshBatchesIntoGBuffer(
+		FRHICommandList& RHICmdList,
+		const FSceneView* View,
+		int32 GPUScenePrimitiveId,
+		TArray<FMeshBatch>& MeshBatches,
+		FVector4 VirtualTexturePhysicalTileCoordinateScaleAndBias,
+		int32 RenderPassIndex,
+		FIntPoint ScratchTilePoolOffset);
+
 	virtual void Finalize(FRDGBuilder& GraphBuilder) override;
 
 	virtual ~FLightmapRenderer();
diff --git a/Engine/Plugins/Experimental/GPULightmass/Source/GPULightmass/Private/Scene/GeometryInterface.h b/Engine/Plugins/Experimental/GPULightmass/Source/GPULightmass/Private/Scene/GeometryInterface.h
--- a/Engine/Plugins/Experimental/GPULightmass/Source/GPULightmass/Private/Scene/GeometryInterface.h
+++ b/Engine/Plugins/Experimental/GPULightmass/Source/GPULightmass/Private/Scene/GeometryInterface.h
@@ -91,6 +91,11 @@ public:
 		return this->Elements[GeometryInstanceRef.GetElementIdChecked()];
 	}
 
+	bool Contains(const FGeometryInstanceRenderStateRef& GeometryInstanceRef) const
+	{
+		return &GeometryInstanceRef.Collection == this;
+	}
+
 	virtual FGeometryRenderState& Get(int32 ElementId) override
 	{
 		return this->Elements[ElementId];
diff --git a/Engine/Plugins/Experimental/GPULightmass/Source/GPULightmass/Private/Scene/InstancedStaticMesh.cpp b/Engine/Plugins/Experimental/GPULightmass/Source/GPULightmass/Private/Scene/InstancedStaticMesh.cpp
--- a/Engine/Plugins/Experimental/GPULightmass/Source/GPULightmass/Private/Scene/InstancedStaticMesh.cpp
+++ b/Engine/Plugins/Experimental/GPULightmass/Source/GPULightmass/Private/Scene/InstancedStaticMesh.cpp
@@ -165,7 +165,6 @@ TArray<FMeshBatch> FInstanceGroupRenderState::GetMeshBatchesForGBufferRendering(
 		MeshBatchElement.PrimitiveIdMode = PrimID_DynamicPrimitiveShaderData;
 
 		MeshBatchElement.InstancedLODIndex = LODIndex;
-		MeshBatchElement.PrimitiveUniformBuffer = UniformBuffer;
 		MeshBatch.LODIndex = LODIndex;
 		MeshBatch.SegmentIndex = SectionIndex;
 		MeshBatch.CastShadow = bCastShadow && Section.bCastShadow;
@@ -175,7 +174,7 @@ TArray<FMeshBatch> FInstanceGroupRenderState::GetMeshBatchesForGBufferRendering(
 		{
 			MeshBatch.MaterialRenderProxy = ComponentUObject->GetMaterial(Section.MaterialIndex)->GetRenderProxy();
 
-			if (CoordsForCulling.MipLevel == -1)
+			if (true || CoordsForCulling.MipLevel == -1)
 			{
 				// No culling, should be for ray tracing scene
 				MeshBatchElement.UserIndex = 0;
diff --git a/Engine/Plugins/Experimental/GPULightmass/Source/GPULightmass/Private/Scene/InstancedStaticMesh.h b/Engine/Plugins/Experimental/GPULightmass/Source/GPULightmass/Private/Scene/InstancedStaticMesh.h
--- a/Engine/Plugins/Experimental/GPULightmass/Source/GPULightmass/Private/Scene/InstancedStaticMesh.h
+++ b/Engine/Plugins/Experimental/GPULightmass/Source/GPULightmass/Private/Scene/InstancedStaticMesh.h
@@ -18,7 +18,6 @@ public:
 	FBufferRHIRef InstanceOriginBuffer;
 	FBufferRHIRef InstanceTransformBuffer;
 	FBufferRHIRef InstanceLightmapBuffer;
-	TUniformBufferRef<FPrimitiveUniformShaderParameters> UniformBuffer;
 
 	TArray<FIntPoint> LODPerInstanceLightmapSize;
 
diff --git a/Engine/Plugins/Experimental/GPULightmass/Source/GPULightmass/Private/Scene/Landscape.cpp b/Engine/Plugins/Experimental/GPULightmass/Source/GPULightmass/Private/Scene/Landscape.cpp
--- a/Engine/Plugins/Experimental/GPULightmass/Source/GPULightmass/Private/Scene/Landscape.cpp
+++ b/Engine/Plugins/Experimental/GPULightmass/Source/GPULightmass/Private/Scene/Landscape.cpp
@@ -110,7 +110,6 @@ TArray<FMeshBatch> FLandscapeRenderState::GetMeshBatchesForGBufferRendering(int3
 	BatchElementParams.CurrentLOD = LODIndex;
 
 	BatchElement.UserData = &BatchElementParams;
-	BatchElement.PrimitiveUniformBuffer = UniformBuffer;
 	BatchElement.IndexBuffer = SharedBuffers->IndexBuffers[LODIndex];
 	BatchElement.NumPrimitives = FMath::Square((SubsectionSizeVerts >> LODIndex) - 1) * FMath::Square(NumSubsections) * 2;
 	BatchElement.FirstIndex = 0;
diff --git a/Engine/Plugins/Experimental/GPULightmass/Source/GPULightmass/Private/Scene/Landscape.h b/Engine/Plugins/Experimental/GPULightmass/Source/GPULightmass/Private/Scene/Landscape.h
--- a/Engine/Plugins/Experimental/GPULightmass/Source/GPULightmass/Private/Scene/Landscape.h
+++ b/Engine/Plugins/Experimental/GPULightmass/Source/GPULightmass/Private/Scene/Landscape.h
@@ -59,8 +59,6 @@ public:
 	TUniquePtr<TUniformBuffer<FLandscapeUniformShaderParameters>> LandscapeUniformShaderParameters;
 	TArray<TUniformBuffer<FLandscapeFixedGridUniformShaderParameters>> LandscapeFixedGridUniformShaderParameters;
 
-	TUniformBufferRef<FPrimitiveUniformShaderParameters> UniformBuffer;
-
 	TArray<FMeshBatch> GetMeshBatchesForGBufferRendering(int32 LODIndex);
 };
 
diff --git a/Engine/Plugins/Experimental/GPULightmass/Source/GPULightmass/Private/Scene/Scene.cpp b/Engine/Plugins/Experimental/GPULightmass/Source/GPULightmass/Private/Scene/Scene.cpp
--- a/Engine/Plugins/Experimental/GPULightmass/Source/GPULightmass/Private/Scene/Scene.cpp
+++ b/Engine/Plugins/Experimental/GPULightmass/Source/GPULightmass/Private/Scene/Scene.cpp
@@ -16,7 +16,6 @@
 #include "LightmapRenderer.h"
 #include "VolumetricLightmap.h"
 #include "GPUScene.h"
-#include "RayTracingDynamicGeometryCollection.h"
 #include "Lightmass/LightmassImportanceVolume.h"
 #include "Logging/MessageLog.h"
 #include "Misc/ConfigCacheIni.h"
@@ -1069,31 +1068,6 @@ void FScene::AddGeometryInstanceFromComponent(UInstancedStaticMeshComponent* InC
 
 		FInstanceGroupRenderStateRef InstanceRenderStateRef = RenderState.InstanceGroupRenderStates.Emplace(MoveTemp(InstanceRenderState));
 
-		InstanceRenderStateRef->UniformBuffer = TUniformBufferRef<FPrimitiveUniformShaderParameters>::CreateUniformBufferImmediate(GetPrimitiveUniformShaderParameters(
-			InstanceRenderStateRef->LocalToWorld,
-			InstanceRenderStateRef->LocalToWorld,
-			InstanceRenderStateRef->ActorPosition,
-			InstanceRenderStateRef->WorldBounds,
-			InstanceRenderStateRef->LocalBounds,
-			InstanceRenderStateRef->LocalBounds,
-			false,
-			false,
-			false,
-			false,
-			false,
-			false,
-			0b111,
-			0,
-			INDEX_NONE,
-			INDEX_NONE,
-			/* bOutputVelocity = */ false,
-			nullptr,
-			/* bCastContactShadow = */ true,
-			INDEX_NONE,
-			0,
-			/* bCastShadow = */ true
-		), UniformBuffer_MultiFrame);
-
 		for (int32 LODIndex = 0; LODIndex < InstanceLightmapRenderStateInitializers.Num(); LODIndex++)
 		{
 			FLightmapRenderState::Initializer& Initializer = InstanceLightmapRenderStateInitializers[LODIndex];
@@ -1189,7 +1163,6 @@ void FScene::RemoveGeometryInstanceFromComponent(UInstancedStaticMeshComponent*
 		[ElementId, &RenderState = RenderState](FRHICommandListImmediate&) mutable
 	{
 		RenderState.InstanceGroupRenderStates.Elements[ElementId].InstancedRenderData->ReleaseResources(nullptr, nullptr);
-		RenderState.InstanceGroupRenderStates.Elements[ElementId].UniformBuffer.SafeRelease();
 
 		for (FLightmapRenderStateRef& Lightmap : RenderState.InstanceGroupRenderStates.Elements[ElementId].LODLightmapRenderStates)
 		{
@@ -1360,31 +1333,6 @@ void FScene::AddGeometryInstanceFromComponent(ULandscapeComponent* InComponent)
 
 		FLandscapeRenderStateRef InstanceRenderStateRef = RenderState.LandscapeRenderStates.Emplace(MoveTemp(InstanceRenderState));
 
-		InstanceRenderStateRef->UniformBuffer = TUniformBufferRef<FPrimitiveUniformShaderParameters>::CreateUniformBufferImmediate(GetPrimitiveUniformShaderParameters(
-			InstanceRenderStateRef->LocalToWorld,
-			InstanceRenderStateRef->LocalToWorld,
-			InstanceRenderStateRef->ActorPosition,
-			InstanceRenderStateRef->WorldBounds,
-			InstanceRenderStateRef->LocalBounds,
-			InstanceRenderStateRef->LocalBounds,
-			false,
-			false,
-			false,
-			false,
-			false,
-			false,
-			0b111,
-			0,
-			INDEX_NONE,
-			INDEX_NONE,
-			/* bOutputVelocity = */ false,
-			nullptr,
-			/* bCastContactShadow = */ true,
-			INDEX_NONE,
-			0,
-			/* bCastShadow = */ true
-		), UniformBuffer_MultiFrame);
-
 		int32 MaxLOD = 0;
 		InstanceRenderStateRef->LandscapeFixedGridUniformShaderParameters.AddDefaulted(MaxLOD + 1);
 		for (int32 LodIndex = 0; LodIndex <= MaxLOD; ++LodIndex)
@@ -1515,89 +1463,9 @@ void FScene::AddGeometryInstanceFromComponent(ULandscapeComponent* InComponent)
 			}
 		}
 
-#if 1
-#if RHI_RAYTRACING
-		if (IsRayTracingEnabled())
-		{
-			// For DynamicGeometryCollection
-			FMaterialRenderProxy::UpdateDeferredCachedUniformExpressions();
-
-			for (int32 SubY = 0; SubY < InstanceRenderStateRef->NumSubsections; SubY++)
-			{
-				for (int32 SubX = 0; SubX < InstanceRenderStateRef->NumSubsections; SubX++)
-				{
-					const int8 SubSectionIdx = SubX + SubY * InstanceRenderStateRef->NumSubsections;
-
-					int32 LodSubsectionSizeVerts = InstanceRenderStateRef->SubsectionSizeVerts;
-					uint32 NumPrimitives = FMath::Square(LodSubsectionSizeVerts - 1) * 2;
-
-					FRayTracingGeometryInitializer GeometryInitializer;
-					GeometryInitializer.IndexBuffer = InstanceRenderStateRef->SharedBuffers->ZeroOffsetIndexBuffers[0]->IndexBufferRHI;
-					GeometryInitializer.TotalPrimitiveCount = NumPrimitives;
-					GeometryInitializer.GeometryType = RTGT_Triangles;
-					GeometryInitializer.bFastBuild = false;
-					GeometryInitializer.bAllowUpdate = false;
-
-					FRayTracingGeometrySegment Segment;
-					Segment.VertexBuffer = nullptr;
-					Segment.VertexBufferStride = sizeof(FVector);
-					Segment.VertexBufferElementType = VET_Float3;
-					Segment.MaxVertices = NumPrimitives*3;
-					Segment.NumPrimitives = NumPrimitives;
-					GeometryInitializer.Segments.Add(Segment);
-
-					InstanceRenderStateRef->SectionRayTracingStates[SubSectionIdx] = MakeUnique<FLandscapeRenderState::FLandscapeSectionRayTracingState>();
-					InstanceRenderStateRef->SectionRayTracingStates[SubSectionIdx]->Geometry.SetInitializer(GeometryInitializer);
-					InstanceRenderStateRef->SectionRayTracingStates[SubSectionIdx]->Geometry.InitResource();
-
-					FRayTracingDynamicGeometryCollection DynamicGeometryCollection;
-
-					TArray<FMeshBatch> MeshBatches = InstanceRenderStateRef->GetMeshBatchesForGBufferRendering(0);
-
-					FLandscapeVertexFactoryMVFParameters UniformBufferParams;
-					UniformBufferParams.SubXY = FIntPoint(SubX, SubY);
-					InstanceRenderStateRef->SectionRayTracingStates[SubSectionIdx]->UniformBuffer = FLandscapeVertexFactoryMVFUniformBufferRef::CreateUniformBufferImmediate(UniformBufferParams, UniformBuffer_SingleFrame);
-
-					FLandscapeBatchElementParams& BatchElementParams = *(FLandscapeBatchElementParams*)MeshBatches[0].Elements[0].UserData;
-					BatchElementParams.LandscapeVertexFactoryMVFUniformBuffer = InstanceRenderStateRef->SectionRayTracingStates[SubSectionIdx]->UniformBuffer;
-
-					MeshBatches[0].Elements[0].IndexBuffer = InstanceRenderStateRef->SharedBuffers->ZeroOffsetIndexBuffers[0];
-					MeshBatches[0].Elements[0].FirstIndex = 0;
-					MeshBatches[0].Elements[0].NumPrimitives = NumPrimitives;
-					MeshBatches[0].Elements[0].MinVertexIndex = 0;
-					MeshBatches[0].Elements[0].MaxVertexIndex = 0;
-
-					FRayTracingDynamicGeometryUpdateParams UpdateParams
-					{
-						MeshBatches,
-						false,
-						(uint32)FMath::Square(LodSubsectionSizeVerts),
-						FMath::Square(LodSubsectionSizeVerts) * (uint32)sizeof(FVector),
-						(uint32)FMath::Square(LodSubsectionSizeVerts - 1) * 2,
-						&InstanceRenderStateRef->SectionRayTracingStates[SubSectionIdx]->Geometry,
-						&InstanceRenderStateRef->SectionRayTracingStates[SubSectionIdx]->RayTracingDynamicVertexBuffer,
-						false
-					};
-
-					DynamicGeometryCollection.AddDynamicMeshBatchForGeometryUpdate(
-						InstanceRenderStateRef->ComponentUObject->GetWorld()->Scene->GetRenderScene(),
-						nullptr,
-						nullptr,
-						UpdateParams,
-						0
-					);
-
-					DynamicGeometryCollection.DispatchUpdates(RHICmdList);
-
-					// Landscape VF doesn't really use the vertex buffer in HitGroupSystemParameters
-					// We can release after all related RHI cmds get dispatched onto the cmd list
-					InstanceRenderStateRef->SectionRayTracingStates[SubSectionIdx]->RayTracingDynamicVertexBuffer.Release();
-				}
-			}
-		}
-#endif
-#endif
 		RenderState.LightmapRenderer->BumpRevision();
+
+		RenderState.CachedRayTracingScene.Reset();
 	});
 
 	bNeedsVoxelization = true;
@@ -1687,6 +1555,8 @@ void FScene::RemoveGeometryInstanceFromComponent(ULandscapeComponent* InComponen
 		RenderState.LandscapeRenderStates.RemoveAt(ElementId);
 
 		RenderState.LightmapRenderer->BumpRevision();
+
+		RenderState.CachedRayTracingScene.Reset();
 	});
 
 	bNeedsVoxelization = true;
diff --git a/Engine/Plugins/Experimental/GPULightmass/Source/GPULightmass/Private/Scene/Scene.h b/Engine/Plugins/Experimental/GPULightmass/Source/GPULightmass/Private/Scene/Scene.h
--- a/Engine/Plugins/Experimental/GPULightmass/Source/GPULightmass/Private/Scene/Scene.h
+++ b/Engine/Plugins/Experimental/GPULightmass/Source/GPULightmass/Private/Scene/Scene.h
@@ -70,6 +70,10 @@ struct FCachedRayTracingSceneData
 	FShaderResourceViewRHIRef InstanceSceneDataBufferSRV;
 	uint32 InstanceSceneDataSOAStride;
 
+	FBufferRHIRef InstanceIdsIdentityBufferRHI;
+	FShaderResourceViewRHIRef InstanceIdsIdentityBufferSRV;
+	TArray<uint32> InstanceDataOriginalOffsets;
+
 	TArray<FRayTracingGeometryInstance> RayTracingGeometryInstances;
 
 	TUniformBufferRef<FViewUniformShaderParameters> CachedViewUniformBuffer;
@@ -107,6 +111,8 @@ public:
 	TUniquePtr<FVolumetricLightmapRenderer> VolumetricLightmapRenderer;
 	TUniquePtr<FIrradianceCache> IrradianceCache;
 
+	int32 GetPrimitiveIdForGPUScene(const FGeometryInstanceRenderStateRef& GeometryInstanceRef) const;
+
 	void SetupRayTracingScene();
 	void DestroyRayTracingScene();
 
diff --git a/Engine/Plugins/Experimental/GPULightmass/Source/GPULightmass/Private/VolumetricLightmap.cpp b/Engine/Plugins/Experimental/GPULightmass/Source/GPULightmass/Private/VolumetricLightmap.cpp
--- a/Engine/Plugins/Experimental/GPULightmass/Source/GPULightmass/Private/VolumetricLightmap.cpp
+++ b/Engine/Plugins/Experimental/GPULightmass/Source/GPULightmass/Private/VolumetricLightmap.cpp
@@ -272,8 +272,10 @@ void FVolumetricLightmapRenderer::VoxelizeScene()
 			};
 		}
 
-		for (FInstanceGroupRenderState& InstanceGroup : Scene->InstanceGroupRenderStates.Elements)
+		for (int32 InstanceGroupIndex = 0; InstanceGroupIndex < Scene->InstanceGroupRenderStates.Elements.Num(); InstanceGroupIndex++)
 		{
+			FInstanceGroupRenderState& InstanceGroup = Scene->InstanceGroupRenderStates.Elements[InstanceGroupIndex];
+
 			bool bIntersectsAnyImportanceVolume = false;
 
 			for (FBox ImportanceVolume : ImportanceVolumes)
@@ -291,13 +293,15 @@ void FVolumetricLightmapRenderer::VoxelizeScene()
 
 			for (auto& MeshBatch : MeshBatches)
 			{
+				MeshBatch.Elements[0].DynamicPrimitiveShaderDataIndex = Scene->StaticMeshInstanceRenderStates.Elements.Num() + InstanceGroupIndex;
 				MeshProcessor.AddMeshBatch(MeshBatch, ~0ull, nullptr);
 			};
 		}
 
-
-		for (FLandscapeRenderState& Landscape : Scene->LandscapeRenderStates.Elements)
+		for (int32 LandscapeIndex = 0; LandscapeIndex < Scene->LandscapeRenderStates.Elements.Num(); LandscapeIndex++)
 		{
+			FLandscapeRenderState& Landscape = Scene->LandscapeRenderStates.Elements[LandscapeIndex];
+
 			bool bIntersectsAnyImportanceVolume = false;
 
 			for (FBox ImportanceVolume : ImportanceVolumes)
@@ -315,6 +319,7 @@ void FVolumetricLightmapRenderer::VoxelizeScene()
 
 			for (auto& MeshBatch : MeshBatches)
 			{
+				MeshBatch.Elements[0].DynamicPrimitiveShaderDataIndex = Scene->StaticMeshInstanceRenderStates.Elements.Num() + Scene->InstanceGroupRenderStates.Elements.Num() + LandscapeIndex;
 				MeshProcessor.AddMeshBatch(MeshBatch, ~0ull, nullptr);
 			};
 		}
diff --git a/Engine/Source/Runtime/Renderer/Private/MeshDrawCommands.h b/Engine/Source/Runtime/Renderer/Private/MeshDrawCommands.h
--- a/Engine/Source/Runtime/Renderer/Private/MeshDrawCommands.h
+++ b/Engine/Source/Runtime/Renderer/Private/MeshDrawCommands.h
@@ -219,7 +219,7 @@ private:
 	void WaitForMeshPassSetupTask() const;
 };
 
-extern void SortAndMergeDynamicPassMeshDrawCommands(
+RENDERER_API extern void SortAndMergeDynamicPassMeshDrawCommands(
 	ERHIFeatureLevel::Type FeatureLevel,
 	FMeshCommandOneFrameArray& VisibleMeshDrawCommands,
 	FDynamicMeshDrawCommandStorage& MeshDrawCommandStorage,
diff --git a/Engine/Source/Runtime/Renderer/Private/RayTracing/RayTracingDynamicGeometry.cpp b/Engine/Source/Runtime/Renderer/Private/RayTracing/RayTracingDynamicGeometry.cpp
--- a/Engine/Source/Runtime/Renderer/Private/RayTracing/RayTracingDynamicGeometry.cpp
+++ b/Engine/Source/Runtime/Renderer/Private/RayTracing/RayTracingDynamicGeometry.cpp
@@ -210,7 +210,7 @@ void FRayTracingDynamicGeometryCollection::AddDynamicMeshBatchForGeometryUpdate(
 
 		int32 DataOffset = 0;
 		FMeshDrawSingleShaderBindings SingleShaderBindings = ShaderBindings.GetSingleShaderBindings(SF_Compute, DataOffset);
-		FMeshPassProcessorRenderState DrawRenderState(Scene->UniformBuffers.ViewUniformBuffer);
+		FMeshPassProcessorRenderState DrawRenderState(View->ViewUniformBuffer);
 		Shader->GetShaderBindings(Scene, Scene->GetFeatureLevel(), PrimitiveSceneProxy, MaterialRenderProxy, Material, DrawRenderState, ShaderElementData, SingleShaderBindings);
 
 		FVertexInputStreamArray DummyArray;
diff --git a/Engine/Source/Runtime/Renderer/Public/MeshPassProcessor.h b/Engine/Source/Runtime/Renderer/Public/MeshPassProcessor.h
--- a/Engine/Source/Runtime/Renderer/Public/MeshPassProcessor.h
+++ b/Engine/Source/Runtime/Renderer/Public/MeshPassProcessor.h
@@ -1799,7 +1799,7 @@ private:
 	EMeshPass::Type PassType;
 };
 
-extern void SubmitMeshDrawCommands(
+RENDERER_API extern void SubmitMeshDrawCommands(
 	const FMeshCommandOneFrameArray& VisibleMeshDrawCommands,
 	const FGraphicsMinimalPipelineStateSet& GraphicsMinimalPipelineStateSet, 
 	FRHIBuffer* PrimitiveIdsBuffer,
