GPULM: temporary fix for on-demand shader compilation crash when bUsedWithStaticLighting flag isn't set on material

diff --git a/Engine/Plugins/Experimental/GPULightmass/Source/GPULightmass/Private/LightmapRayTracing.cpp b/Engine/Plugins/Experimental/GPULightmass/Source/GPULightmass/Private/LightmapRayTracing.cpp
--- a/Engine/Plugins/Experimental/GPULightmass/Source/GPULightmass/Private/LightmapRayTracing.cpp
+++ b/Engine/Plugins/Experimental/GPULightmass/Source/GPULightmass/Private/LightmapRayTracing.cpp
@@ -132,18 +132,6 @@ IMPLEMENT_GLOBAL_SHADER(FStationaryLightShadowTracingRGS, "/Plugin/GPULightmass/
 
 IMPLEMENT_GLOBAL_SHADER(FFirstBounceRayGuidingCDFBuildCS, "/Plugin/GPULightmass/Private/FirstBounceRayGuidingCDFBuild.usf", "FirstBounceRayGuidingCDFBuildCS", SF_Compute);
 
-static TShaderRef<FLightmapMaterialCHS> GetLightmapMaterialHitShader(const FMaterial& RESTRICT MaterialResource, const FVertexFactory* VertexFactory)
-{
-	if (MaterialResource.IsMasked() || MaterialResource.GetBlendMode() == BLEND_Translucent)
-	{
-		return MaterialResource.GetShader<TLightmapMaterialCHS<true>>(VertexFactory->GetType());
-	}
-	else
-	{
-		return MaterialResource.GetShader<TLightmapMaterialCHS<false>>(VertexFactory->GetType());
-	}
-}
-
 bool FLightmapRayTracingMeshProcessor::Process(
 	const FMeshBatch& RESTRICT MeshBatch,
 	uint64 BatchElementMask,
@@ -163,7 +151,24 @@ bool FLightmapRayTracingMeshProcessor::Process(
 		FMeshMaterialShader,
 		FLightmapMaterialCHS> RayTracingShaders;
 
-	RayTracingShaders.RayHitGroupShader = GetLightmapMaterialHitShader(MaterialResource, VertexFactory);
+	FMaterialShaderTypes ShaderTypes;
+
+	if (MaterialResource.IsMasked() || MaterialResource.GetBlendMode() == BLEND_Translucent)
+	{
+		ShaderTypes.AddShaderType<TLightmapMaterialCHS<true>>();
+	}
+	else
+	{
+		ShaderTypes.AddShaderType<TLightmapMaterialCHS<false>>();
+	}
+
+	FMaterialShaders Shaders;
+	if (!MaterialResource.TryGetShaders(ShaderTypes, VertexFactory->GetType(), Shaders))
+	{
+		return false;
+	}
+
+	check(Shaders.TryGetShader(SF_RayHitGroup, RayTracingShaders.RayHitGroupShader));
 
 	PassDrawRenderState.SetBlendState(TStaticBlendState<CW_RGBA, BO_Add, BF_One, BF_One, BO_Add, BF_Zero, BF_One>::GetRHI());
 	PassDrawRenderState.SetDepthStencilState(TStaticDepthStencilState<false, CF_DepthNearOrEqual>::GetRHI());
diff --git a/Engine/Plugins/Experimental/GPULightmass/Source/GPULightmass/Private/LightmapRenderer.cpp b/Engine/Plugins/Experimental/GPULightmass/Source/GPULightmass/Private/LightmapRenderer.cpp
--- a/Engine/Plugins/Experimental/GPULightmass/Source/GPULightmass/Private/LightmapRenderer.cpp
+++ b/Engine/Plugins/Experimental/GPULightmass/Source/GPULightmass/Private/LightmapRenderer.cpp
@@ -814,6 +814,8 @@ void FSceneRenderState::SetupRayTracingScene()
 
 		SCOPED_DRAW_EVENTF(RHICmdList, GPULightmassUpdateRayTracingScene, TEXT("GPULightmass UpdateRayTracingScene %d Instances"), StaticMeshInstanceRenderStates.Elements.Num());
 
+		FMemMark MemMark(FMemStack::Get());
+
 		TArray<FRayTracingGeometryInstance, SceneRenderingAllocator> RayTracingGeometryInstances;
 		RayTracingGeometryInstances.Append(CachedRayTracingScene->RayTracingGeometryInstances);
 
