Fixed compile errors introduced in 16158927

diff --git a/Engine/Plugins/Experimental/GPULightmass/Source/GPULightmass/Private/Scene/Scene.cpp b/Engine/Plugins/Experimental/GPULightmass/Source/GPULightmass/Private/Scene/Scene.cpp
--- a/Engine/Plugins/Experimental/GPULightmass/Source/GPULightmass/Private/Scene/Scene.cpp
+++ b/Engine/Plugins/Experimental/GPULightmass/Source/GPULightmass/Private/Scene/Scene.cpp
@@ -1316,7 +1316,7 @@ void FScene::AddGeometryInstanceFromComponent(ULandscapeComponent* InComponent)
 		{
 			InstanceRenderState.SharedBuffers = new FLandscapeSharedBuffers(
 				InstanceRenderState.SharedBuffersKey, Initializer.SubsectionSizeQuads, Initializer.NumSubsections,
-				FeatureLevel, /*NumOcclusionVertices*/ 0);
+				FeatureLevel);
 
 			FLandscapeComponentSceneProxy::SharedBuffersMap.Add(InstanceRenderState.SharedBuffersKey, InstanceRenderState.SharedBuffers);
 
diff --git a/Engine/Source/Runtime/Renderer/Private/ReflectionEnvironmentRealTimeCapture.cpp b/Engine/Source/Runtime/Renderer/Private/ReflectionEnvironmentRealTimeCapture.cpp
--- a/Engine/Source/Runtime/Renderer/Private/ReflectionEnvironmentRealTimeCapture.cpp
+++ b/Engine/Source/Runtime/Renderer/Private/ReflectionEnvironmentRealTimeCapture.cpp
@@ -16,6 +16,7 @@
 #include "VolumetricCloudProxy.h"
 #include "FogRendering.h"
 #include "GPUScene.h"
+#include "ScreenPass.h"
 
 #if WITH_EDITOR
 #include "CanvasTypes.h"
diff --git a/Engine/Source/Runtime/Renderer/Private/Strata/Strata.cpp b/Engine/Source/Runtime/Renderer/Private/Strata/Strata.cpp
--- a/Engine/Source/Runtime/Renderer/Private/Strata/Strata.cpp
+++ b/Engine/Source/Runtime/Renderer/Private/Strata/Strata.cpp
@@ -10,6 +10,7 @@
 #include "UniformBuffer.h"
 #include "SceneTextureParameters.h"
 #include "StrataDefinitions.h"
+#include "ScreenPass.h"
 
 //PRAGMA_DISABLE_OPTIMIZATION
 
