GPULM: Fix BatchedTilesSRV lifetime issue

diff --git a/Engine/Plugins/Experimental/GPULightmass/Source/GPULightmass/Private/LightmapRenderer.cpp b/Engine/Plugins/Experimental/GPULightmass/Source/GPULightmass/Private/LightmapRenderer.cpp
--- a/Engine/Plugins/Experimental/GPULightmass/Source/GPULightmass/Private/LightmapRenderer.cpp
+++ b/Engine/Plugins/Experimental/GPULightmass/Source/GPULightmass/Private/LightmapRenderer.cpp
@@ -2711,7 +2711,7 @@ void FLightmapRenderer::Finalize(FRDGBuilder& GraphBuilder)
 			CreateInfo.ResourceArray = &GPUBatchedTileRequests.BatchedTilesDesc;
 
 			GPUBatchedTileRequests.BatchedTilesBuffer = RHICreateStructuredBuffer(sizeof(FGPUTileDescription), GPUBatchedTileRequests.BatchedTilesDesc.GetResourceDataSize(), BUF_Dynamic | BUF_ShaderResource, CreateInfo);
-			GPUBatchedTileRequests.BatchedTilesSRV = RHICreateShaderResourceView(GPUBatchedTileRequests.BatchedTilesBuffer);
+			GPUBatchedTileRequests.BatchedTilesSRV = HoldReference(GraphBuilder, RHICreateShaderResourceView(GPUBatchedTileRequests.BatchedTilesBuffer));
 		}
 
 		{
