GPULM: Port Volumetric Lightmap voxelization to RDG & fix static mesh PrimitiveUniformShaderParameters

diff --git a/Engine/Plugins/Experimental/GPULightmass/Shaders/Private/VolumetricLightmapVoxelization.usf b/Engine/Plugins/Experimental/GPULightmass/Shaders/Private/VolumetricLightmapVoxelization.usf
--- a/Engine/Plugins/Experimental/GPULightmass/Shaders/Private/VolumetricLightmapVoxelization.usf
+++ b/Engine/Plugins/Experimental/GPULightmass/Shaders/Private/VolumetricLightmapVoxelization.usf
@@ -13,8 +13,6 @@
 
 #define SOFTWARE_CONSERVATIVE_RASTERIZATION 1
 
-RWTexture3D<uint> VoxelizeVolume;
-
 struct FVLMVoxelizationVSToGS
 {
 	float4 WorldPosition : SV_POSITION;
diff --git a/Engine/Plugins/Experimental/GPULightmass/Source/GPULightmass/Private/Scene/Scene.cpp b/Engine/Plugins/Experimental/GPULightmass/Source/GPULightmass/Private/Scene/Scene.cpp
--- a/Engine/Plugins/Experimental/GPULightmass/Source/GPULightmass/Private/Scene/Scene.cpp
+++ b/Engine/Plugins/Experimental/GPULightmass/Source/GPULightmass/Private/Scene/Scene.cpp
@@ -790,7 +790,7 @@ void FScene::AddGeometryInstanceFromComponent(UStaticMeshComponent* InComponent)
 	{
 		FStaticMeshInstanceRenderStateRef InstanceRenderStateRef = RenderState.StaticMeshInstanceRenderStates.Emplace(MoveTemp(InstanceRenderState));
 
-		FPrimitiveUniformShaderParameters PrimitiveUniformShaderParameters =
+		InstanceRenderStateRef->PrimitiveUniformShaderParameters =
 			FPrimitiveUniformShaderParametersBuilder{}
 			.Defaults()
 				.LocalToWorld(InstanceRenderStateRef->LocalToWorld)
diff --git a/Engine/Plugins/Experimental/GPULightmass/Source/GPULightmass/Private/VolumetricLightmap.cpp b/Engine/Plugins/Experimental/GPULightmass/Source/GPULightmass/Private/VolumetricLightmap.cpp
--- a/Engine/Plugins/Experimental/GPULightmass/Source/GPULightmass/Private/VolumetricLightmap.cpp
+++ b/Engine/Plugins/Experimental/GPULightmass/Source/GPULightmass/Private/VolumetricLightmap.cpp
@@ -94,6 +94,12 @@ FPrecomputedVolumetricLightmap* FVolumetricLightmapRenderer::GetPrecomputedVolum
 	return &VolumetricLightmap;
 }
 
+BEGIN_SHADER_PARAMETER_STRUCT(FVoxelizeMeshPassParameters, )
+	SHADER_PARAMETER_RDG_UNIFORM_BUFFER(FVLMVoxelizationParams, PassUniformBuffer)
+	SHADER_PARAMETER_RDG_UNIFORM_BUFFER(FInstanceCullingGlobalUniforms, InstanceCulling)
+	RENDER_TARGET_BINDING_SLOTS()
+END_SHADER_PARAMETER_STRUCT()
+
 void FVolumetricLightmapRenderer::VoxelizeScene()
 {
 	for (int32 MipLevel = 0; MipLevel < VoxelizationVolumeMips.Num(); MipLevel++)
@@ -176,184 +182,194 @@ void FVolumetricLightmapRenderer::VoxelizeScene()
 	FBox CubeVolume(VolumeMin, VolumeMin + FVector(FMath::Max3(VolumeSize.X, VolumeSize.Y, VolumeSize.Z)));
 	int32 CubeMaxDim = FMath::Max3(IndirectionTextureDimensions.X, IndirectionTextureDimensions.Y, IndirectionTextureDimensions.Z);
 
-	FVLMVoxelizationParams VLMVoxelizationParams;
-	VLMVoxelizationParams.VolumeCenter = CubeVolume.GetCenter();
-	VLMVoxelizationParams.VolumeExtent = CubeVolume.GetExtent();
-	VLMVoxelizationParams.VolumeMaxDim = CubeMaxDim;
-	VLMVoxelizationParams.VoxelizeVolume = VoxelizationVolumeMips[0]->GetRenderTargetItem().UAV;
-	VLMVoxelizationParams.IndirectionTexture = IndirectionTexture->GetRenderTargetItem().UAV;
-	FVLMVoxelizationUniformBufferRef PassUniformBuffer = FVLMVoxelizationUniformBufferRef::CreateUniformBufferImmediate(VLMVoxelizationParams, UniformBuffer_SingleFrame);
-
 	FRHICommandListImmediate& RHICmdList = FRHICommandListExecutor::GetImmediateCommandList();
 
-	FUniformBufferStaticBindings GlobalUniformBuffers;
-	GlobalUniformBuffers.AddUniformBuffer(PassUniformBuffer);
-	RHICmdList.SetStaticUniformBuffers(GlobalUniformBuffers);
+	FMemMark MemMark(FMemStack::Get());
+
+	FRDGBuilder GraphBuilder(RHICmdList);
+
+	FVLMVoxelizationParams* VLMVoxelizationParams = GraphBuilder.AllocParameters<FVLMVoxelizationParams>();
+	VLMVoxelizationParams->VolumeCenter = CubeVolume.GetCenter();
+	VLMVoxelizationParams->VolumeExtent = CubeVolume.GetExtent();
+	VLMVoxelizationParams->VolumeMaxDim = CubeMaxDim;
+	VLMVoxelizationParams->VoxelizeVolume = VoxelizationVolumeMips[0]->GetRenderTargetItem().UAV;
+	VLMVoxelizationParams->IndirectionTexture = IndirectionTexture->GetRenderTargetItem().UAV;
+	TRDGUniformBufferRef<FVLMVoxelizationParams> PassUniformBuffer = GraphBuilder.CreateUniformBuffer(VLMVoxelizationParams);
 
 	for (int32 MipLevel = 0; MipLevel < VoxelizationVolumeMips.Num(); MipLevel++)
 	{
 		FGlobalShaderMap* GlobalShaderMap = GetGlobalShaderMap(GMaxRHIFeatureLevel);
 
-		TShaderMapRef<FClearVolumeCS> ComputeShader(GlobalShaderMap);
+		FClearVolumeCS::FParameters* Parameters = GraphBuilder.AllocParameters<FClearVolumeCS::FParameters>();
+		Parameters->VolumeSize = VoxelizationVolumeMips[MipLevel]->GetDesc().GetSize();
+		Parameters->VoxelizeVolume = VoxelizationVolumeMips[MipLevel]->GetRenderTargetItem().UAV;
 
-		FClearVolumeCS::FParameters Parameters;
-		Parameters.VolumeSize = VoxelizationVolumeMips[MipLevel]->GetDesc().GetSize();
-		Parameters.VoxelizeVolume = VoxelizationVolumeMips[MipLevel]->GetRenderTargetItem().UAV;
-
-		FComputeShaderUtils::Dispatch(RHICmdList, ComputeShader, Parameters, FComputeShaderUtils::GetGroupCount(VoxelizationVolumeMips[MipLevel]->GetDesc().GetSize(), FIntVector(4)));
-
-		RHICmdList.Transition(FRHITransitionInfo(VoxelizationVolumeMips[MipLevel]->GetRenderTargetItem().UAV, ERHIAccess::Unknown, ERHIAccess::UAVCompute));
+		TShaderMapRef<FClearVolumeCS> ComputeShader(GlobalShaderMap);
+		FComputeShaderUtils::AddPass(
+			GraphBuilder,
+			RDG_EVENT_NAME("ClearVolume"),
+			ComputeShader,
+			Parameters,
+			FComputeShaderUtils::GetGroupCount(VoxelizationVolumeMips[MipLevel]->GetDesc().GetSize(), FIntVector(4)));
 	}
 
+	for (FBox ImportanceVolume : ImportanceVolumes)
 	{
-		for (FBox ImportanceVolume : ImportanceVolumes)
-		{
-			FGlobalShaderMap* GlobalShaderMap = GetGlobalShaderMap(GMaxRHIFeatureLevel);
-
-			TShaderMapRef<FVoxelizeImportanceVolumeCS> ComputeShader(GlobalShaderMap);
-
-			FVoxelizeImportanceVolumeCS::FParameters Parameters;
-			Parameters.VolumeSize = VoxelizationVolumeMips[0]->GetDesc().GetSize();
-			Parameters.ImportanceVolumeMin = ImportanceVolume.Min;
-			Parameters.ImportanceVolumeMax = ImportanceVolume.Max;
-			Parameters.VLMVoxelizationParams = PassUniformBuffer;
-			Parameters.VoxelizeVolume = VoxelizationVolumeMips[0]->GetRenderTargetItem().UAV;
-
-			FComputeShaderUtils::Dispatch(RHICmdList, ComputeShader, Parameters, FComputeShaderUtils::GetGroupCount(VoxelizationVolumeMips[0]->GetDesc().GetSize(), FIntVector(4)));
-		}
+		FGlobalShaderMap* GlobalShaderMap = GetGlobalShaderMap(GMaxRHIFeatureLevel);
 
-		RHICmdList.Transition(FRHITransitionInfo(VoxelizationVolumeMips[0]->GetRenderTargetItem().UAV, ERHIAccess::UAVCompute, ERHIAccess::UAVGraphics));
+		TShaderMapRef<FVoxelizeImportanceVolumeCS> ComputeShader(GlobalShaderMap);
+
+		FVoxelizeImportanceVolumeCS::FParameters* Parameters = GraphBuilder.AllocParameters<FVoxelizeImportanceVolumeCS::FParameters>();
+		Parameters->VolumeSize = VoxelizationVolumeMips[0]->GetDesc().GetSize();
+		Parameters->ImportanceVolumeMin = ImportanceVolume.Min;
+		Parameters->ImportanceVolumeMax = ImportanceVolume.Max;
+		Parameters->VLMVoxelizationParams = PassUniformBuffer;
+		Parameters->VoxelizeVolume = VoxelizationVolumeMips[0]->GetRenderTargetItem().UAV;
+
+		FComputeShaderUtils::AddPass(
+			GraphBuilder,
+			RDG_EVENT_NAME("VoxelizeImportanceVolume"),
+			ComputeShader,
+			Parameters,
+			FComputeShaderUtils::GetGroupCount(VoxelizationVolumeMips[0]->GetDesc().GetSize(), FIntVector(4)));
 	}
 
 	// Setup ray tracing scene with LOD 0
 	Scene->SetupRayTracingScene();
 
-	FMemMark Mark(FMemStack::Get());
-
-	FRHIRenderPassInfo RPInfo(FRHIRenderPassInfo::NoRenderTargets);
-
-	RHICmdList.BeginRenderPass(RPInfo, TEXT("VolumetricLightmapVoxelization"));
-
-	RHICmdList.SetViewport(0, 0, 0, CubeMaxDim, CubeMaxDim, 1);
-
-	SCOPED_DRAW_EVENTF(RHICmdList, GPULightmassVoxelizeScene, TEXT("GPULightmass VoxelizeScene"));
+	auto* PassParameters = GraphBuilder.AllocParameters<FVoxelizeMeshPassParameters>();
+	PassParameters->PassUniformBuffer = GraphBuilder.CreateUniformBuffer(VLMVoxelizationParams);
+	PassParameters->InstanceCulling = FInstanceCullingContext::CreateDummyInstanceCullingUniformBuffer(GraphBuilder);
 
-	DrawDynamicMeshPass(
-		*Scene->ReferenceView,
-		RHICmdList,
-		[&Scene = Scene, View = Scene->ReferenceView.Get(), PassUniformBuffer, ImportanceVolumes = ImportanceVolumes](FDynamicPassMeshDrawListContext* DynamicMeshPassContext)
+	GraphBuilder.AddPass(
+		RDG_EVENT_NAME("VLM Mesh Voxelization"),
+		PassParameters,
+		ERDGPassFlags::Raster | ERDGPassFlags::NeverCull,
+		[this, CubeMaxDim, PassUniformBuffer, PassParameters](FRHICommandList& RHICmdList)
 	{
-		FVLMVoxelizationMeshProcessor MeshProcessor(nullptr, View, DynamicMeshPassContext, PassUniformBuffer);
+		RHICmdList.SetViewport(0, 0, 0, CubeMaxDim, CubeMaxDim, 1);
 
-		for (int32 InstanceIndex = 0; InstanceIndex < Scene->StaticMeshInstanceRenderStates.Elements.Num(); InstanceIndex++)
-		{
-			FStaticMeshInstanceRenderState& Instance = Scene->StaticMeshInstanceRenderStates.Elements[InstanceIndex];
+		SCOPED_DRAW_EVENTF(RHICmdList, GPULightmassVoxelizeScene, TEXT("GPULightmass VoxelizeScene"));
 
-			bool bIntersectsAnyImportanceVolume = false;
+		DrawDynamicMeshPass(
+			*Scene->ReferenceView,
+			RHICmdList,
+			[&Scene = Scene, View = Scene->ReferenceView.Get(), ImportanceVolumes = ImportanceVolumes](FDynamicPassMeshDrawListContext* DynamicMeshPassContext)
+		{
+			FVLMVoxelizationMeshProcessor MeshProcessor(nullptr, View, DynamicMeshPassContext);
 
-			for (FBox ImportanceVolume : ImportanceVolumes)
+			for (int32 InstanceIndex = 0; InstanceIndex < Scene->StaticMeshInstanceRenderStates.Elements.Num(); InstanceIndex++)
 			{
-				if (Instance.WorldBounds.GetBox().Intersect(ImportanceVolume))
+				FStaticMeshInstanceRenderState& Instance = Scene->StaticMeshInstanceRenderStates.Elements[InstanceIndex];
+
+				bool bIntersectsAnyImportanceVolume = false;
+
+				for (FBox ImportanceVolume : ImportanceVolumes)
 				{
-					bIntersectsAnyImportanceVolume = true;
-					break;
+					if (Instance.WorldBounds.GetBox().Intersect(ImportanceVolume))
+					{
+						bIntersectsAnyImportanceVolume = true;
+						break;
+					}
 				}
-			}
 
-			if (!bIntersectsAnyImportanceVolume) continue;
+				if (!bIntersectsAnyImportanceVolume) continue;
 
-			TArray<FMeshBatch> MeshBatches = Instance.GetMeshBatchesForGBufferRendering(0);
+				TArray<FMeshBatch> MeshBatches = Instance.GetMeshBatchesForGBufferRendering(0);
 
-			for (auto& MeshBatch : MeshBatches)
-			{
-				MeshBatch.Elements[0].DynamicPrimitiveShaderDataIndex = InstanceIndex;
-				MeshProcessor.AddMeshBatch(MeshBatch, ~0ull, nullptr);
-			};
-		}
+				for (auto& MeshBatch : MeshBatches)
+				{
+					MeshBatch.Elements[0].DynamicPrimitiveShaderDataIndex = InstanceIndex;
+					MeshProcessor.AddMeshBatch(MeshBatch, ~0ull, nullptr);
+				};
+			}
 
-		for (int32 InstanceGroupIndex = 0; InstanceGroupIndex < Scene->InstanceGroupRenderStates.Elements.Num(); InstanceGroupIndex++)
-		{
-			FInstanceGroupRenderState& InstanceGroup = Scene->InstanceGroupRenderStates.Elements[InstanceGroupIndex];
+			for (int32 InstanceGroupIndex = 0; InstanceGroupIndex < Scene->InstanceGroupRenderStates.Elements.Num(); InstanceGroupIndex++)
+			{
+				FInstanceGroupRenderState& InstanceGroup = Scene->InstanceGroupRenderStates.Elements[InstanceGroupIndex];
 
-			bool bIntersectsAnyImportanceVolume = false;
+				bool bIntersectsAnyImportanceVolume = false;
 
-			for (FBox ImportanceVolume : ImportanceVolumes)
-			{
-				if (InstanceGroup.WorldBounds.GetBox().Intersect(ImportanceVolume))
+				for (FBox ImportanceVolume : ImportanceVolumes)
 				{
-					bIntersectsAnyImportanceVolume = true;
-					break;
+					if (InstanceGroup.WorldBounds.GetBox().Intersect(ImportanceVolume))
+					{
+						bIntersectsAnyImportanceVolume = true;
+						break;
+					}
 				}
-			}
 
-			if (!bIntersectsAnyImportanceVolume) continue;
+				if (!bIntersectsAnyImportanceVolume) continue;
 
-			TArray<FMeshBatch> MeshBatches = InstanceGroup.GetMeshBatchesForGBufferRendering(0, FTileVirtualCoordinates{});
+				TArray<FMeshBatch> MeshBatches = InstanceGroup.GetMeshBatchesForGBufferRendering(0, FTileVirtualCoordinates{});
 
-			for (auto& MeshBatch : MeshBatches)
-			{
-				MeshBatch.Elements[0].DynamicPrimitiveShaderDataIndex = Scene->StaticMeshInstanceRenderStates.Elements.Num() + InstanceGroupIndex;
-				MeshProcessor.AddMeshBatch(MeshBatch, ~0ull, nullptr);
-			};
-		}
+				for (auto& MeshBatch : MeshBatches)
+				{
+					MeshBatch.Elements[0].DynamicPrimitiveShaderDataIndex = Scene->StaticMeshInstanceRenderStates.Elements.Num() + InstanceGroupIndex;
+					MeshProcessor.AddMeshBatch(MeshBatch, ~0ull, nullptr);
+				};
+			}
 
-		for (int32 LandscapeIndex = 0; LandscapeIndex < Scene->LandscapeRenderStates.Elements.Num(); LandscapeIndex++)
-		{
-			FLandscapeRenderState& Landscape = Scene->LandscapeRenderStates.Elements[LandscapeIndex];
+			for (int32 LandscapeIndex = 0; LandscapeIndex < Scene->LandscapeRenderStates.Elements.Num(); LandscapeIndex++)
+			{
+				FLandscapeRenderState& Landscape = Scene->LandscapeRenderStates.Elements[LandscapeIndex];
 
-			bool bIntersectsAnyImportanceVolume = false;
+				bool bIntersectsAnyImportanceVolume = false;
 
-			for (FBox ImportanceVolume : ImportanceVolumes)
-			{
-				if (Landscape.WorldBounds.GetBox().Intersect(ImportanceVolume))
+				for (FBox ImportanceVolume : ImportanceVolumes)
 				{
-					bIntersectsAnyImportanceVolume = true;
-					break;
+					if (Landscape.WorldBounds.GetBox().Intersect(ImportanceVolume))
+					{
+						bIntersectsAnyImportanceVolume = true;
+						break;
+					}
 				}
-			}
 
-			if (!bIntersectsAnyImportanceVolume) continue;
+				if (!bIntersectsAnyImportanceVolume) continue;
 
-			TArray<FMeshBatch> MeshBatches = Landscape.GetMeshBatchesForGBufferRendering(0);
+				TArray<FMeshBatch> MeshBatches = Landscape.GetMeshBatchesForGBufferRendering(0);
 
-			for (auto& MeshBatch : MeshBatches)
-			{
-				MeshBatch.Elements[0].DynamicPrimitiveShaderDataIndex = Scene->StaticMeshInstanceRenderStates.Elements.Num() + Scene->InstanceGroupRenderStates.Elements.Num() + LandscapeIndex;
-				MeshProcessor.AddMeshBatch(MeshBatch, ~0ull, nullptr);
-			};
-		}
+				for (auto& MeshBatch : MeshBatches)
+				{
+					MeshBatch.Elements[0].DynamicPrimitiveShaderDataIndex = Scene->StaticMeshInstanceRenderStates.Elements.Num() + Scene->InstanceGroupRenderStates.Elements.Num() + LandscapeIndex;
+					MeshProcessor.AddMeshBatch(MeshBatch, ~0ull, nullptr);
+				};
+			}
+		});
 	});
 
-	RHICmdList.EndRenderPass();
-
-	RHICmdList.Transition(FRHITransitionInfo(VoxelizationVolumeMips[0]->GetRenderTargetItem().UAV, ERHIAccess::UAVGraphics, ERHIAccess::UAVCompute));
-
 	{
 		FGlobalShaderMap* GlobalShaderMap = GetGlobalShaderMap(GMaxRHIFeatureLevel);
 
 		TShaderMapRef<FDilateVolumeCS> ComputeShader(GlobalShaderMap);
 
-		FDilateVolumeCS::FParameters Parameters;
-		Parameters.VolumeSize = VoxelizationVolumeMips[0]->GetDesc().GetSize();
-		Parameters.VoxelizeVolume = VoxelizationVolumeMips[0]->GetRenderTargetItem().UAV;
+		FDilateVolumeCS::FParameters* Parameters = GraphBuilder.AllocParameters<FDilateVolumeCS::FParameters>();
+		Parameters->VolumeSize = VoxelizationVolumeMips[0]->GetDesc().GetSize();
+		Parameters->VoxelizeVolume = VoxelizationVolumeMips[0]->GetRenderTargetItem().UAV;
 
-		FComputeShaderUtils::Dispatch(RHICmdList, ComputeShader, Parameters, FComputeShaderUtils::GetGroupCount(VoxelizationVolumeMips[0]->GetDesc().GetSize(), FIntVector(4)));
+		FComputeShaderUtils::AddPass(
+			GraphBuilder,
+			RDG_EVENT_NAME("DilateVolume"),
+			ComputeShader,
+			Parameters,
+			FComputeShaderUtils::GetGroupCount(VoxelizationVolumeMips[0]->GetDesc().GetSize(), FIntVector(4)));
 	}
 
-	RHICmdList.Transition(FRHITransitionInfo(VoxelizationVolumeMips[0]->GetRenderTargetItem().UAV, ERHIAccess::UAVCompute, ERHIAccess::UAVCompute));
-
 	for (int32 MipLevel = 1; MipLevel < VoxelizationVolumeMips.Num(); MipLevel++)
 	{
 		TShaderMapRef<FDownsampleVolumeCS> ComputeShader(GetGlobalShaderMap(GMaxRHIFeatureLevel));
 
-		FDownsampleVolumeCS::FParameters Parameters;
-		Parameters.bIsHighestMip = (MipLevel == VoxelizationVolumeMips.Num() - 1) ? 1 : 0;
-		Parameters.VoxelizeVolume = VoxelizationVolumeMips[MipLevel]->GetRenderTargetItem().UAV;
-		Parameters.VoxelizeVolumePrevMip = VoxelizationVolumeMips[MipLevel - 1]->GetRenderTargetItem().UAV;
-
-		FComputeShaderUtils::Dispatch(RHICmdList, ComputeShader, Parameters, VoxelizationVolumeMips[MipLevel]->GetDesc().GetSize());
-
-		RHICmdList.Transition(FRHITransitionInfo(VoxelizationVolumeMips[MipLevel]->GetRenderTargetItem().UAV, ERHIAccess::UAVCompute, ERHIAccess::UAVCompute));
+		FDownsampleVolumeCS::FParameters* Parameters = GraphBuilder.AllocParameters<FDownsampleVolumeCS::FParameters>();
+		Parameters->bIsHighestMip = (MipLevel == VoxelizationVolumeMips.Num() - 1) ? 1 : 0;
+		Parameters->VoxelizeVolume = VoxelizationVolumeMips[MipLevel]->GetRenderTargetItem().UAV;
+		Parameters->VoxelizeVolumePrevMip = VoxelizationVolumeMips[MipLevel - 1]->GetRenderTargetItem().UAV;
+
+		FComputeShaderUtils::AddPass(
+			GraphBuilder,
+			RDG_EVENT_NAME("DownsampleVolume"),
+			ComputeShader,
+			Parameters,
+			VoxelizationVolumeMips[MipLevel]->GetDesc().GetSize());
 	}
 
 	{
@@ -369,17 +385,21 @@ void FVolumetricLightmapRenderer::VoxelizeScene()
 
 		TShaderMapRef<FCountNumBricksCS> ComputeShader(GlobalShaderMap);
 
-		FCountNumBricksCS::FParameters Parameters;
-		Parameters.VolumeSize = VoxelizationVolumeMips[MipLevel]->GetDesc().GetSize();
-		Parameters.VoxelizeVolume = VoxelizationVolumeMips[MipLevel]->GetRenderTargetItem().UAV;
-		Parameters.BrickAllocatorParameters = BrickAllocatorParameters.UAV;
-
-		FComputeShaderUtils::Dispatch(RHICmdList, ComputeShader, Parameters, FComputeShaderUtils::GetGroupCount(VoxelizationVolumeMips[MipLevel]->GetDesc().GetSize(), FIntVector(4)));
-
-		RHICmdList.Transition(FRHITransitionInfo(VoxelizationVolumeMips[MipLevel]->GetRenderTargetItem().UAV, ERHIAccess::UAVCompute, ERHIAccess::UAVCompute));
-		RHICmdList.Transition(FRHITransitionInfo(BrickAllocatorParameters.UAV, ERHIAccess::Unknown, ERHIAccess::UAVCompute));
+		FCountNumBricksCS::FParameters* Parameters = GraphBuilder.AllocParameters<FCountNumBricksCS::FParameters>();
+		Parameters->VolumeSize = VoxelizationVolumeMips[MipLevel]->GetDesc().GetSize();
+		Parameters->VoxelizeVolume = VoxelizationVolumeMips[MipLevel]->GetRenderTargetItem().UAV;
+		Parameters->BrickAllocatorParameters = BrickAllocatorParameters.UAV;
+
+		FComputeShaderUtils::AddPass(
+			GraphBuilder,
+			RDG_EVENT_NAME("CountNumBricks"),
+			ComputeShader,
+			Parameters,
+			FComputeShaderUtils::GetGroupCount(VoxelizationVolumeMips[MipLevel]->GetDesc().GetSize(), FIntVector(4)));
 	}
 
+	GraphBuilder.Execute();
+
 	FRHIGPUBufferReadback NumBricksReadback(TEXT("NumBricksReadback"));
 	NumBricksReadback.EnqueueCopy(RHICmdList, BrickAllocatorParameters.Buffer);
 	RHICmdList.BlockUntilGPUIdle();
diff --git a/Engine/Plugins/Experimental/GPULightmass/Source/GPULightmass/Private/VolumetricLightmapVoxelization.h b/Engine/Plugins/Experimental/GPULightmass/Source/GPULightmass/Private/VolumetricLightmapVoxelization.h
--- a/Engine/Plugins/Experimental/GPULightmass/Source/GPULightmass/Private/VolumetricLightmapVoxelization.h
+++ b/Engine/Plugins/Experimental/GPULightmass/Source/GPULightmass/Private/VolumetricLightmapVoxelization.h
@@ -17,7 +17,6 @@
 BEGIN_GLOBAL_SHADER_PARAMETER_STRUCT(FVLMVoxelizationParams, )
 	SHADER_PARAMETER(FVector4, VolumeCenter)
 	SHADER_PARAMETER(FVector4, VolumeExtent)
-	SHADER_PARAMETER(FIntVector, VolumeSize)
 	SHADER_PARAMETER(int32, VolumeMaxDim)
 	SHADER_PARAMETER_UAV(RWTexture3D<uint>, VoxelizeVolume)
 	SHADER_PARAMETER_UAV(RWTexture3D<uint4>, IndirectionTexture)
@@ -36,7 +35,6 @@ protected:
 	FVLMVoxelizationVS(const ShaderMetaType::CompiledShaderInitializerType& Initializer)
 		: FMeshMaterialShader(Initializer)
 	{
-		PassUniformBuffer.Bind(Initializer.ParameterMap, FVLMVoxelizationParams::StaticStructMetadata.GetShaderVariableName());
 	}
 
 	FVLMVoxelizationVS()
@@ -59,7 +57,6 @@ protected:
 	FVLMVoxelizationGS(const ShaderMetaType::CompiledShaderInitializerType& Initializer)
 		: FMeshMaterialShader(Initializer)
 	{
-		PassUniformBuffer.Bind(Initializer.ParameterMap, FVLMVoxelizationParams::StaticStructMetadata.GetShaderVariableName());
 	}
 
 	FVLMVoxelizationGS()
@@ -82,7 +79,6 @@ protected:
 	FVLMVoxelizationPS(const ShaderMetaType::CompiledShaderInitializerType& Initializer)
 		: FMeshMaterialShader(Initializer)
 	{
-		PassUniformBuffer.Bind(Initializer.ParameterMap, FVLMVoxelizationParams::StaticStructMetadata.GetShaderVariableName());
 	}
 
 	FVLMVoxelizationPS()
@@ -99,9 +95,9 @@ protected:
 class FVLMVoxelizationMeshProcessor : public FMeshPassProcessor
 {
 public:
-	FVLMVoxelizationMeshProcessor(const FScene* InScene, const FSceneView* InView, FMeshPassDrawListContext* InDrawListContext, FRHIUniformBuffer* InPassUniformBuffer)
+	FVLMVoxelizationMeshProcessor(const FScene* InScene, const FSceneView* InView, FMeshPassDrawListContext* InDrawListContext)
 		: FMeshPassProcessor(InScene, GMaxRHIFeatureLevel, InView, InDrawListContext)
-		, DrawRenderState(*InView, InPassUniformBuffer)
+		, DrawRenderState(*InView)
 	{
 		DrawRenderState.SetDepthStencilState(TStaticDepthStencilState<false, CF_Always>::GetRHI());
 		DrawRenderState.SetBlendState(TStaticBlendState<>::GetRHI());
@@ -211,7 +207,7 @@ class FVoxelizeImportanceVolumeCS : public FGlobalShader
 		SHADER_PARAMETER(FVector3f, ImportanceVolumeMin)
 		SHADER_PARAMETER(FVector3f, ImportanceVolumeMax)
 		SHADER_PARAMETER_UAV(RWTexture3D<uint>, VoxelizeVolume)
-		SHADER_PARAMETER_STRUCT_REF(FVLMVoxelizationParams, VLMVoxelizationParams)
+		SHADER_PARAMETER_RDG_UNIFORM_BUFFER(FVLMVoxelizationParams, VLMVoxelizationParams)
 	END_SHADER_PARAMETER_STRUCT()
 };
 
