GPULM: Fix GPU crash when LODs have different number of material sections

diff --git a/Engine/Plugins/Experimental/GPULightmass/Source/GPULightmass/Private/LightmapRenderer.cpp b/Engine/Plugins/Experimental/GPULightmass/Source/GPULightmass/Private/LightmapRenderer.cpp
--- a/Engine/Plugins/Experimental/GPULightmass/Source/GPULightmass/Private/LightmapRenderer.cpp
+++ b/Engine/Plugins/Experimental/GPULightmass/Source/GPULightmass/Private/LightmapRenderer.cpp
@@ -287,13 +287,6 @@ void FCachedRayTracingSceneData::SetupViewUniformBufferFromSceneRenderState(FSce
 			ConservativeLightmapEntriesNum += Landscape.LODLightmapRenderStates.Num();
 		}
 
-		for (int32 InstanceGroupIndex = 0; InstanceGroupIndex < Scene.InstanceGroupRenderStates.Elements.Num(); InstanceGroupIndex++)
-		{
-			FInstanceGroupRenderState& InstanceGroup = Scene.InstanceGroupRenderStates.Elements[InstanceGroupIndex];
-			LightmapSceneDataStartOffsets[InstanceGroupIndex] = ConservativeLightmapEntriesNum;
-			ConservativeLightmapEntriesNum += InstanceGroup.LODLightmapRenderStates.Num();
-		}
-
 		LightmapSceneData.AddZeroed(ConservativeLightmapEntriesNum);
 	}
 
@@ -316,8 +309,8 @@ void FCachedRayTracingSceneData::SetupViewUniformBufferFromSceneRenderState(FSce
 				FPrimitiveUniformShaderParameters PrimitiveUniformShaderParameters = Instance.PrimitiveUniformShaderParameters;
 				PrimitiveUniformShaderParameters.LightmapDataIndex = LightmapSceneDataStartOffsets[InstanceIndex];
 				PrimitiveUniformShaderParameters.LightmapUVIndex = 0; // TODO: LightmapUVIndex
-				PrimitiveUniformShaderParameters.InstanceSceneDataOffset = InstanceIndex;
-				PrimitiveUniformShaderParameters.NumInstanceSceneDataEntries = 1;
+				PrimitiveUniformShaderParameters.InstanceDataOffset = InstanceIndex;
+				PrimitiveUniformShaderParameters.NumInstanceDataEntries = 1;
 				PrimitiveSceneData[InstanceIndex] = FPrimitiveSceneShaderData(PrimitiveUniformShaderParameters);
 
 				InstanceDataOriginalOffsets[InstanceIndex] = InstanceIndex;
@@ -377,8 +370,8 @@ void FCachedRayTracingSceneData::SetupViewUniformBufferFromSceneRenderState(FSce
 					.LocalBounds(InstanceGroup.LocalBounds)
 					.LightingChannelMask(0b111)
 					.LightmapDataIndex(LightmapSceneDataStartOffsets[PrimitiveId])
-					.InstanceSceneDataOffset(InstanceSceneData.Num())
-					.NumInstanceSceneDataEntries(NumInstancesThisGroup)
+					.InstanceDataOffset(InstanceSceneData.Num())
+					.NumInstanceDataEntries(NumInstancesThisGroup)
 					.CastContactShadow(true)
 					.CastShadow(true)
 				.Build();
@@ -440,8 +433,8 @@ void FCachedRayTracingSceneData::SetupViewUniformBufferFromSceneRenderState(FSce
 					.LocalBounds(Landscape.LocalBounds)
 					.LightingChannelMask(0b111)
 					.LightmapDataIndex(LightmapSceneDataStartOffsets[PrimitiveId])
-					.InstanceSceneDataOffset(InstanceSceneData.Num())
-					.NumInstanceSceneDataEntries(1)
+					.InstanceDataOffset(InstanceSceneData.Num())
+					.NumInstanceDataEntries(1)
 					.CastContactShadow(true)
 					.CastShadow(true)
 				.Build();
@@ -524,8 +517,8 @@ void FCachedRayTracingSceneData::SetupViewUniformBufferFromSceneRenderState(FSce
 			TRACE_CPUPROFILER_EVENT_SCOPE(InstanceSceneData);
 
 			TResourceArray<FVector4> InstanceSceneDataSOA;
-			InstanceSceneDataSOA.AddZeroed(FInstanceSceneShaderData::DataStrideInFloat4s * InstanceSceneData.Num());
-			for (int32 ArrayIndex = 0; ArrayIndex < FInstanceSceneShaderData::DataStrideInFloat4s; ArrayIndex++)
+			InstanceSceneDataSOA.AddZeroed(FInstanceSceneShaderData::InstanceDataStrideInFloat4s * InstanceSceneData.Num());
+			for (int32 ArrayIndex = 0; ArrayIndex < FInstanceSceneShaderData::InstanceDataStrideInFloat4s; ArrayIndex++)
 			{
 				for (int32 DataIndex = 0; DataIndex < InstanceSceneData.Num(); DataIndex++)
 				{
@@ -536,7 +529,7 @@ void FCachedRayTracingSceneData::SetupViewUniformBufferFromSceneRenderState(FSce
 			FRHIResourceCreateInfo CreateInfo(TEXT("InstanceSceneDataBuffer"), &InstanceSceneDataSOA);
 			if (InstanceSceneDataSOA.GetResourceDataSize() == 0)
 			{
-				InstanceSceneDataSOA.AddZeroed(FInstanceSceneShaderData::DataStrideInFloat4s);
+				InstanceSceneDataSOA.AddZeroed(FInstanceSceneShaderData::InstanceDataStrideInFloat4s);
 			}
 
 			InstanceSceneDataBufferRHI = RHICreateStructuredBuffer(sizeof(FVector4), InstanceSceneDataSOA.GetResourceDataSize(), BUF_Static | BUF_ShaderResource, CreateInfo);
@@ -549,22 +542,52 @@ void FCachedRayTracingSceneData::SetupViewUniformBufferFromSceneRenderState(FSce
 	}
 }
 
-void FCachedRayTracingSceneData::SetupFromSceneRenderState(FSceneRenderState& Scene, int32 LODIndex)
+void FCachedRayTracingSceneData::SetupFromSceneRenderState(FSceneRenderState& Scene)
 {
 #if RHI_RAYTRACING
-	SceneLODIndex = LODIndex;
+	RayTracingGeometryInstancesPerLOD.AddDefaulted(MAX_STATIC_MESH_LODS);
+	VisibleRayTracingMeshCommandsPerLOD.AddDefaulted(MAX_STATIC_MESH_LODS);
 
 	FMaterialRenderProxy::UpdateDeferredCachedUniformExpressions();
 
+	for (int32 LODIndex = 0; LODIndex < MAX_STATIC_MESH_LODS; LODIndex++)
 	{
-		RayTracingGeometryInstances.Reserve(Scene.StaticMeshInstanceRenderStates.Elements.Num());
-		StaticMeshInstanceIndices.AddZeroed(Scene.StaticMeshInstanceRenderStates.Elements.Num());
+		bool bShouldIncludeThisLODLevel = false;
 
-		for (int32 StaticMeshIndex = 0; StaticMeshIndex < Scene.StaticMeshInstanceRenderStates.Elements.Num();  StaticMeshIndex++)
+		for (int32 StaticMeshIndex = 0; StaticMeshIndex < Scene.StaticMeshInstanceRenderStates.Elements.Num(); StaticMeshIndex++)
 		{
 			FStaticMeshInstanceRenderState& Instance = Scene.StaticMeshInstanceRenderStates.Elements[StaticMeshIndex];
+			int32 LODIndexToUse = FMath::Clamp(LODIndex, Instance.ClampedMinLOD, Instance.RenderData->LODResources.Num() - 1);
+			if (LODIndexToUse == LODIndex)
+			{
+				bShouldIncludeThisLODLevel = true;
+				break;
+			}
+		}
 
-			int32 LODIndexToUse = FMath::Clamp(SceneLODIndex, Instance.ClampedMinLOD, Instance.RenderData->LODResources.Num() - 1);
+		for (int32 InstanceGroupIndex = 0; InstanceGroupIndex < Scene.InstanceGroupRenderStates.Elements.Num(); InstanceGroupIndex++)
+		{
+			FInstanceGroupRenderState& InstanceGroup = Scene.InstanceGroupRenderStates.Elements[InstanceGroupIndex];
+			int32 LODIndexToUse = FMath::Clamp(LODIndex, 0, InstanceGroup.ComponentUObject->GetStaticMesh()->GetRenderData()->LODResources.Num() - 1);
+			if (LODIndexToUse == LODIndex)
+			{
+				bShouldIncludeThisLODLevel = true;
+				break;
+			}
+		}
+
+		if (!bShouldIncludeThisLODLevel)
+		{
+			continue;
+		}
+
+		RayTracingGeometryInstancesPerLOD[LODIndex].Reserve(Scene.StaticMeshInstanceRenderStates.Elements.Num());
+
+		for (int32 StaticMeshIndex = 0; StaticMeshIndex < Scene.StaticMeshInstanceRenderStates.Elements.Num(); StaticMeshIndex++)
+		{
+			FStaticMeshInstanceRenderState& Instance = Scene.StaticMeshInstanceRenderStates.Elements[StaticMeshIndex];
+
+			int32 LODIndexToUse = FMath::Clamp(LODIndex, Instance.ClampedMinLOD, Instance.RenderData->LODResources.Num() - 1);
 
 			TArray<FMeshBatch> MeshBatches = Instance.GetMeshBatchesForGBufferRendering(LODIndexToUse);
 
@@ -588,9 +611,8 @@ void FCachedRayTracingSceneData::SetupFromSceneRenderState(FSceneRenderState& Sc
 
 			if (!bAllSegmentsUnlit)
 			{
-				int32 InstanceIndex = RayTracingGeometryInstances.AddDefaulted(1);
-				StaticMeshInstanceIndices[StaticMeshIndex] = InstanceIndex;
-				FRayTracingGeometryInstance& RayTracingInstance = RayTracingGeometryInstances[InstanceIndex];
+				int32 InstanceIndex = RayTracingGeometryInstancesPerLOD[LODIndex].AddDefaulted(1);
+				FRayTracingGeometryInstance& RayTracingInstance = RayTracingGeometryInstancesPerLOD[LODIndex][InstanceIndex];
 				RayTracingInstance.GeometryRHI = Instance.RenderData->LODResources[LODIndexToUse].RayTracingGeometry.RayTracingGeometryRHI;
 				RayTracingInstance.Transforms = MakeArrayView(&Instance.LocalToWorld, 1);
 				RayTracingInstance.NumTransforms = 1;
@@ -603,93 +625,81 @@ void FCachedRayTracingSceneData::SetupFromSceneRenderState(FSceneRenderState& Sc
 
 				for (int32 SegmentIndex = 0; SegmentIndex < MeshBatches.Num(); SegmentIndex++)
 				{
-					FFullyCachedRayTracingMeshCommandContext CommandContext(MeshCommandStorage, VisibleRayTracingMeshCommands, SegmentIndex, InstanceIndex);
+					FFullyCachedRayTracingMeshCommandContext CommandContext(MeshCommandStorage, VisibleRayTracingMeshCommandsPerLOD[LODIndex], SegmentIndex, InstanceIndex);
 					FMeshPassProcessorRenderState PassDrawRenderState(CachedViewUniformBuffer, CachedViewUniformBuffer);
 					FLightmapRayTracingMeshProcessor RayTracingMeshProcessor(&CommandContext, PassDrawRenderState);
 
 					RayTracingMeshProcessor.AddMeshBatch(MeshBatches[SegmentIndex], 1, nullptr);
 				}
 			}
-			else
-			{
-				StaticMeshInstanceIndices[StaticMeshIndex] = INDEX_NONE;
-			}
 		}
 
-		RayTracingGeometryInstances.Reserve(Scene.InstanceGroupRenderStates.Elements.Num());
-		InstanceGroupInstanceIndices.AddZeroed(Scene.InstanceGroupRenderStates.Elements.Num());
+		RayTracingGeometryInstancesPerLOD[LODIndex].Reserve(Scene.InstanceGroupRenderStates.Elements.Num());
 
+		for (int32 InstanceGroupIndex = 0; InstanceGroupIndex < Scene.InstanceGroupRenderStates.Elements.Num(); InstanceGroupIndex++)
 		{
-			for (int32 InstanceGroupIndex = 0; InstanceGroupIndex < Scene.InstanceGroupRenderStates.Elements.Num(); InstanceGroupIndex++)
-			{
-				FInstanceGroupRenderState& InstanceGroup = Scene.InstanceGroupRenderStates.Elements[InstanceGroupIndex];
-
-				int32 LODIndexToUse = FMath::Clamp(SceneLODIndex, 0, InstanceGroup.ComponentUObject->GetStaticMesh()->GetRenderData()->LODResources.Num() - 1);
-
-				TArray<FMeshBatch> MeshBatches = InstanceGroup.GetMeshBatchesForGBufferRendering(LODIndexToUse, FTileVirtualCoordinates{});
+			FInstanceGroupRenderState& InstanceGroup = Scene.InstanceGroupRenderStates.Elements[InstanceGroupIndex];
 
-				bool bAllSegmentsUnlit = true;
-				bool bAllSegmentsOpaque = true;
-				bool bAnySegmentsCastShadow = false;
-				uint8 InstanceMask = 0;
+			int32 LODIndexToUse = FMath::Clamp(LODIndex, 0, InstanceGroup.ComponentUObject->GetStaticMesh()->GetRenderData()->LODResources.Num() - 1);
 
-				for (int32 SegmentIndex = 0; SegmentIndex < MeshBatches.Num(); SegmentIndex++)
-				{
-					const FMaterialRenderProxy* FallbackMaterialRenderProxyPtr = nullptr;
-					const FMaterial& Material = MeshBatches[SegmentIndex].MaterialRenderProxy->GetMaterialWithFallback(GMaxRHIFeatureLevel, FallbackMaterialRenderProxyPtr);
+			TArray<FMeshBatch> MeshBatches = InstanceGroup.GetMeshBatchesForGBufferRendering(LODIndexToUse, FTileVirtualCoordinates{});
 
-					bAllSegmentsUnlit &= Material.GetShadingModels().HasOnlyShadingModel(MSM_Unlit) || !MeshBatches[SegmentIndex].CastShadow;
-					bAllSegmentsOpaque &= Material.GetBlendMode() == EBlendMode::BLEND_Opaque;
-					bAnySegmentsCastShadow |= MeshBatches[SegmentIndex].CastRayTracedShadow && Material.CastsRayTracedShadows();
-					InstanceMask |= ComputeBlendModeMask(Material.GetBlendMode());
-				}
+			bool bAllSegmentsUnlit = true;
+			bool bAllSegmentsOpaque = true;
+			bool bAnySegmentsCastShadow = false;
+			uint8 InstanceMask = 0;
 
-				InstanceMask |= bAnySegmentsCastShadow ? RAY_TRACING_MASK_SHADOW : 0;
+			for (int32 SegmentIndex = 0; SegmentIndex < MeshBatches.Num(); SegmentIndex++)
+			{
+				const FMaterialRenderProxy* FallbackMaterialRenderProxyPtr = nullptr;
+				const FMaterial& Material = MeshBatches[SegmentIndex].MaterialRenderProxy->GetMaterialWithFallback(GMaxRHIFeatureLevel, FallbackMaterialRenderProxyPtr);
 
-				if (!bAllSegmentsUnlit)
-				{
-					int32 InstanceIndex = RayTracingGeometryInstances.AddDefaulted(1);
-					InstanceGroupInstanceIndices[InstanceGroupIndex] = InstanceIndex;
-					FRayTracingGeometryInstance& RayTracingInstance = RayTracingGeometryInstances[InstanceIndex];
-					RayTracingInstance.GeometryRHI = InstanceGroup.ComponentUObject->GetStaticMesh()->GetRenderData()->LODResources[LODIndexToUse].RayTracingGeometry.RayTracingGeometryRHI;
+				bAllSegmentsUnlit &= Material.GetShadingModels().HasOnlyShadingModel(MSM_Unlit) || !MeshBatches[SegmentIndex].CastShadow;
+				bAllSegmentsOpaque &= Material.GetBlendMode() == EBlendMode::BLEND_Opaque;
+				bAnySegmentsCastShadow |= MeshBatches[SegmentIndex].CastRayTracedShadow && Material.CastsRayTracedShadows();
+				InstanceMask |= ComputeBlendModeMask(Material.GetBlendMode());
+			}
 
-					const int32 NumInstances = (int32)InstanceGroup.InstancedRenderData->PerInstanceRenderData->InstanceBuffer.GetNumInstances();
-					TArrayView<FMatrix> NewTransforms = MakeArrayView(
-						OwnedRayTracingInstanceTransforms.Add_GetRef(TUniquePtr<FMatrix>(new FMatrix[NumInstances])).Get(),
-						NumInstances);
+			InstanceMask |= bAnySegmentsCastShadow ? RAY_TRACING_MASK_SHADOW : 0;
 
-					for (int32 InstanceIdx = 0; InstanceIdx < NumInstances; InstanceIdx++)
-					{
-						FRenderTransform Transform;
-						InstanceGroup.InstancedRenderData->PerInstanceRenderData->InstanceBuffer.GetInstanceTransform(InstanceIdx, Transform);
+			if (!bAllSegmentsUnlit)
+			{
+				int32 InstanceIndex = RayTracingGeometryInstancesPerLOD[LODIndex].AddDefaulted(1);
+				FRayTracingGeometryInstance& RayTracingInstance = RayTracingGeometryInstancesPerLOD[LODIndex][InstanceIndex];
+				RayTracingInstance.GeometryRHI = InstanceGroup.ComponentUObject->GetStaticMesh()->GetRenderData()->LODResources[LODIndexToUse].RayTracingGeometry.RayTracingGeometryRHI;
 
-						FMatrix InstanceTransform = Transform.ToMatrix() * InstanceGroup.LocalToWorld;
-						NewTransforms[InstanceIdx] = InstanceTransform;
-					}
+				const int32 NumInstances = (int32)InstanceGroup.InstancedRenderData->PerInstanceRenderData->InstanceBuffer.GetNumInstances();
+				TArrayView<FMatrix> NewTransforms = MakeArrayView(
+					OwnedRayTracingInstanceTransforms.Add_GetRef(TUniquePtr<FMatrix>(new FMatrix[NumInstances])).Get(),
+					NumInstances);
 
-					RayTracingInstance.Transforms = NewTransforms;
-					RayTracingInstance.NumTransforms = NumInstances;
+				for (int32 InstanceIdx = 0; InstanceIdx < NumInstances; InstanceIdx++)
+				{
+					FRenderTransform Transform;
+					InstanceGroup.InstancedRenderData->PerInstanceRenderData->InstanceBuffer.GetInstanceTransform(InstanceIdx, Transform);
 
-					RayTracingInstance.DefaultUserData = (uint32)(Scene.StaticMeshInstanceRenderStates.Elements.Num() + InstanceGroupIndex);
-					RayTracingInstance.Mask = InstanceMask;
+					FMatrix InstanceTransform = Transform.ToMatrix() * InstanceGroup.LocalToWorld;
+					NewTransforms[InstanceIdx] = InstanceTransform;
+				}
 
-					if (bAllSegmentsOpaque)
-					{
-						RayTracingInstance.Flags |= ERayTracingInstanceFlags::ForceOpaque;
-					}
+				RayTracingInstance.Transforms = NewTransforms;
+				RayTracingInstance.NumTransforms = NumInstances;
 
-					for (int32 SegmentIndex = 0; SegmentIndex < MeshBatches.Num(); SegmentIndex++)
-					{
-						FFullyCachedRayTracingMeshCommandContext CommandContext(MeshCommandStorage, VisibleRayTracingMeshCommands, SegmentIndex, InstanceIndex);
-						FMeshPassProcessorRenderState PassDrawRenderState(CachedViewUniformBuffer, CachedViewUniformBuffer);
-						FLightmapRayTracingMeshProcessor RayTracingMeshProcessor(&CommandContext, PassDrawRenderState);
+				RayTracingInstance.DefaultUserData = (uint32)(Scene.StaticMeshInstanceRenderStates.Elements.Num() + InstanceGroupIndex);
+				RayTracingInstance.Mask = InstanceMask;
 
-						RayTracingMeshProcessor.AddMeshBatch(MeshBatches[SegmentIndex], 1, nullptr);
-					}
+				if (bAllSegmentsOpaque)
+				{
+					RayTracingInstance.Flags |= ERayTracingInstanceFlags::ForceOpaque;
 				}
-				else
+
+				for (int32 SegmentIndex = 0; SegmentIndex < MeshBatches.Num(); SegmentIndex++)
 				{
-					InstanceGroupInstanceIndices[InstanceGroupIndex] = INDEX_NONE;
+					FFullyCachedRayTracingMeshCommandContext CommandContext(MeshCommandStorage, VisibleRayTracingMeshCommandsPerLOD[LODIndex], SegmentIndex, InstanceIndex);
+					FMeshPassProcessorRenderState PassDrawRenderState(CachedViewUniformBuffer, CachedViewUniformBuffer);
+					FLightmapRayTracingMeshProcessor RayTracingMeshProcessor(&CommandContext, PassDrawRenderState);
+
+					RayTracingMeshProcessor.AddMeshBatch(MeshBatches[SegmentIndex], 1, nullptr);
 				}
 			}
 		}
@@ -699,37 +709,6 @@ void FCachedRayTracingSceneData::SetupFromSceneRenderState(FSceneRenderState& Sc
 #endif // RHI_RAYTRACING
 }
 
-void FCachedRayTracingSceneData::UpdateLODRayTracingGeometry(FSceneRenderState& Scene, int32 LODIndex)
-{
-	SceneLODIndex = LODIndex;
-
-	for (int32 StaticMeshIndex = 0; StaticMeshIndex < Scene.StaticMeshInstanceRenderStates.Elements.Num(); StaticMeshIndex++)
-	{
-		FStaticMeshInstanceRenderState& Instance = Scene.StaticMeshInstanceRenderStates.Elements[StaticMeshIndex];
-
-		int32 LODIndexToUse = FMath::Clamp(SceneLODIndex, Instance.ClampedMinLOD, Instance.RenderData->LODResources.Num() - 1);
-
-		if (StaticMeshInstanceIndices[StaticMeshIndex] != INDEX_NONE)
-		{
-			FRayTracingGeometryInstance& RayTracingInstance = RayTracingGeometryInstances[StaticMeshInstanceIndices[StaticMeshIndex]];
-			RayTracingInstance.GeometryRHI = Instance.RenderData->LODResources[LODIndexToUse].RayTracingGeometry.RayTracingGeometryRHI;
-		}
-	}
-
-	for (int32 InstanceGroupIndex = 0; InstanceGroupIndex < Scene.InstanceGroupRenderStates.Elements.Num(); InstanceGroupIndex++)
-	{
-		FInstanceGroupRenderState& InstanceGroup = Scene.InstanceGroupRenderStates.Elements[InstanceGroupIndex];
-
-		int32 LODIndexToUse = FMath::Clamp(SceneLODIndex, 0, InstanceGroup.ComponentUObject->GetStaticMesh()->GetRenderData()->LODResources.Num() - 1);
-
-		if (InstanceGroupInstanceIndices[InstanceGroupIndex] != INDEX_NONE)
-		{
-			FRayTracingGeometryInstance& RayTracingInstance = RayTracingGeometryInstances[InstanceGroupInstanceIndices[InstanceGroupIndex]];
-			RayTracingInstance.GeometryRHI = InstanceGroup.ComponentUObject->GetStaticMesh()->GetRenderData()->LODResources[LODIndexToUse].RayTracingGeometry.RayTracingGeometryRHI;
-		}
-	}
-}
-
 void FSceneRenderState::SetupRayTracingScene(int32 LODIndex)
 {
 	TRACE_CPUPROFILER_EVENT_SCOPE(SetupRayTracingScene);
@@ -746,14 +725,10 @@ void FSceneRenderState::SetupRayTracingScene(int32 LODIndex)
 	{
 		CachedRayTracingScene = MakeUnique<FCachedRayTracingSceneData>();
 		CachedRayTracingScene->SetupViewUniformBufferFromSceneRenderState(*this);
-		CachedRayTracingScene->SetupFromSceneRenderState(*this, LODIndex);
+		CachedRayTracingScene->SetupFromSceneRenderState(*this);
 		
 		CalculateDistributionPrefixSumForAllLightmaps();
 	}
-	else if (CachedRayTracingScene->SceneLODIndex != LODIndex)
-	{
-		CachedRayTracingScene->UpdateLODRayTracingGeometry(*this, LODIndex);
-	}
 
 #if 0 // Debug: verify cached ray tracing scene has up-to-date shader bindings
 	TUniquePtr<FCachedRayTracingSceneData> VerificationRayTracingScene = MakeUnique<FCachedRayTracingSceneData>();
@@ -815,7 +790,7 @@ void FSceneRenderState::SetupRayTracingScene(int32 LODIndex)
 				TVC_MAX,
 				*View.CachedViewUniformShaderParameters);
 
-			View.CachedViewUniformShaderParameters->InstanceSceneDataSOAStride = CachedRayTracingScene->InstanceSceneDataSOAStride;
+			View.CachedViewUniformShaderParameters->InstanceDataSOAStride = CachedRayTracingScene->InstanceSceneDataSOAStride;
 
 			if (LightSceneRenderState.SkyLight.IsSet())
 			{
@@ -847,7 +822,7 @@ void FSceneRenderState::SetupRayTracingScene(int32 LODIndex)
 		FMemMark MemMark(FMemStack::Get());
 
 		TArray<FRayTracingGeometryInstance, SceneRenderingAllocator> RayTracingGeometryInstances;
-		RayTracingGeometryInstances.Append(CachedRayTracingScene->RayTracingGeometryInstances);
+		RayTracingGeometryInstances.Append(CachedRayTracingScene->RayTracingGeometryInstancesPerLOD[LODIndex]);
 
 		int32 LandscapeStartOffset = RayTracingGeometryInstances.Num();
 		for (FLandscapeRenderState& Landscape : LandscapeRenderStates.Elements)
@@ -958,7 +933,7 @@ void FSceneRenderState::SetupRayTracingScene(int32 LODIndex)
 						RayTracingInstance.GeometryRHI = Landscape.SectionRayTracingStates[SubSectionIdx]->Geometry.RayTracingGeometryRHI;
 						RayTracingInstance.Transforms = MakeArrayView(&FMatrix::Identity, 1);
 						RayTracingInstance.NumTransforms = 1;
-						RayTracingInstance.DefaultUserData = (uint32)InstanceIndex;
+						RayTracingInstance.DefaultUserData = (uint32)(StaticMeshInstanceRenderStates.Elements.Num() + InstanceGroupRenderStates.Elements.Num() + LandscapeIndex);
 
 						TArray<FMeshBatch> MeshBatches = Landscape.GetMeshBatchesForGBufferRendering(0);
 
@@ -1072,7 +1047,7 @@ void FSceneRenderState::SetupRayTracingScene(int32 LODIndex)
 			{
 				TRACE_CPUPROFILER_EVENT_SCOPE(SetRayTracingShaderBindings);
 
-				for (const FVisibleRayTracingMeshCommand VisibleMeshCommand : CachedRayTracingScene->VisibleRayTracingMeshCommands)
+				for (const FVisibleRayTracingMeshCommand VisibleMeshCommand : CachedRayTracingScene->VisibleRayTracingMeshCommandsPerLOD[LODIndex])
 				{
 					const FRayTracingMeshCommand& MeshCommand = *VisibleMeshCommand.RayTracingMeshCommand;
 
diff --git a/Engine/Plugins/Experimental/GPULightmass/Source/GPULightmass/Private/Scene/Scene.h b/Engine/Plugins/Experimental/GPULightmass/Source/GPULightmass/Private/Scene/Scene.h
--- a/Engine/Plugins/Experimental/GPULightmass/Source/GPULightmass/Private/Scene/Scene.h
+++ b/Engine/Plugins/Experimental/GPULightmass/Source/GPULightmass/Private/Scene/Scene.h
@@ -58,9 +58,7 @@ private:
 
 struct FCachedRayTracingSceneData
 {
-	int32 SceneLODIndex = 0;
-
-	TArray<FVisibleRayTracingMeshCommand> VisibleRayTracingMeshCommands;
+	TArray<TArray<FVisibleRayTracingMeshCommand>> VisibleRayTracingMeshCommandsPerLOD;
 	TChunkedArray<FRayTracingMeshCommand> MeshCommandStorage;
 
 	FBufferRHIRef PrimitiveSceneDataBufferRHI;
@@ -77,17 +75,13 @@ struct FCachedRayTracingSceneData
 	FShaderResourceViewRHIRef InstanceIdsIdentityBufferSRV;
 	TArray<uint32> InstanceDataOriginalOffsets;
 
-	TArray<FRayTracingGeometryInstance> RayTracingGeometryInstances;
+	TArray<TArray<FRayTracingGeometryInstance>> RayTracingGeometryInstancesPerLOD;
 	TArray<TUniquePtr<FMatrix>> OwnedRayTracingInstanceTransforms;
 
 	TUniformBufferRef<FViewUniformShaderParameters> CachedViewUniformBuffer;
 
-	TArray<int32> StaticMeshInstanceIndices;
-	TArray<int32> InstanceGroupInstanceIndices;
-
 	void SetupViewUniformBufferFromSceneRenderState(class FSceneRenderState& Scene);
-	void SetupFromSceneRenderState(class FSceneRenderState& Scene, int32 LODIndex);
-	void UpdateLODRayTracingGeometry(FSceneRenderState& Scene, int32 LODIndex);
+	void SetupFromSceneRenderState(class FSceneRenderState& Scene);
 };
 
 class FSceneRenderState
