Fix accidentally reverted changes

diff --git a/Engine/Plugins/Experimental/GPULightmass/Source/GPULightmass/Private/LightmapRenderer.cpp b/Engine/Plugins/Experimental/GPULightmass/Source/GPULightmass/Private/LightmapRenderer.cpp
--- a/Engine/Plugins/Experimental/GPULightmass/Source/GPULightmass/Private/LightmapRenderer.cpp
+++ b/Engine/Plugins/Experimental/GPULightmass/Source/GPULightmass/Private/LightmapRenderer.cpp
@@ -309,8 +309,8 @@ void FCachedRayTracingSceneData::SetupViewUniformBufferFromSceneRenderState(FSce
 				FPrimitiveUniformShaderParameters PrimitiveUniformShaderParameters = Instance.PrimitiveUniformShaderParameters;
 				PrimitiveUniformShaderParameters.LightmapDataIndex = LightmapSceneDataStartOffsets[InstanceIndex];
 				PrimitiveUniformShaderParameters.LightmapUVIndex = 0; // TODO: LightmapUVIndex
-				PrimitiveUniformShaderParameters.InstanceDataOffset = InstanceIndex;
-				PrimitiveUniformShaderParameters.NumInstanceDataEntries = 1;
+				PrimitiveUniformShaderParameters.InstanceSceneDataOffset = InstanceIndex;
+				PrimitiveUniformShaderParameters.NumInstanceSceneDataEntries = 1;
 				PrimitiveSceneData[InstanceIndex] = FPrimitiveSceneShaderData(PrimitiveUniformShaderParameters);
 
 				InstanceDataOriginalOffsets[InstanceIndex] = InstanceIndex;
@@ -370,8 +370,8 @@ void FCachedRayTracingSceneData::SetupViewUniformBufferFromSceneRenderState(FSce
 					.LocalBounds(InstanceGroup.LocalBounds)
 					.LightingChannelMask(0b111)
 					.LightmapDataIndex(LightmapSceneDataStartOffsets[PrimitiveId])
-					.InstanceDataOffset(InstanceSceneData.Num())
-					.NumInstanceDataEntries(NumInstancesThisGroup)
+					.InstanceSceneDataOffset(InstanceSceneData.Num())
+					.NumInstanceSceneDataEntries(NumInstancesThisGroup)
 					.CastContactShadow(true)
 					.CastShadow(true)
 				.Build();
@@ -433,8 +433,8 @@ void FCachedRayTracingSceneData::SetupViewUniformBufferFromSceneRenderState(FSce
 					.LocalBounds(Landscape.LocalBounds)
 					.LightingChannelMask(0b111)
 					.LightmapDataIndex(LightmapSceneDataStartOffsets[PrimitiveId])
-					.InstanceDataOffset(InstanceSceneData.Num())
-					.NumInstanceDataEntries(1)
+					.InstanceSceneDataOffset(InstanceSceneData.Num())
+					.NumInstanceSceneDataEntries(1)
 					.CastContactShadow(true)
 					.CastShadow(true)
 				.Build();
@@ -517,8 +517,8 @@ void FCachedRayTracingSceneData::SetupViewUniformBufferFromSceneRenderState(FSce
 			TRACE_CPUPROFILER_EVENT_SCOPE(InstanceSceneData);
 
 			TResourceArray<FVector4> InstanceSceneDataSOA;
-			InstanceSceneDataSOA.AddZeroed(FInstanceSceneShaderData::InstanceDataStrideInFloat4s * InstanceSceneData.Num());
-			for (int32 ArrayIndex = 0; ArrayIndex < FInstanceSceneShaderData::InstanceDataStrideInFloat4s; ArrayIndex++)
+			InstanceSceneDataSOA.AddZeroed(FInstanceSceneShaderData::DataStrideInFloat4s * InstanceSceneData.Num());
+			for (int32 ArrayIndex = 0; ArrayIndex < FInstanceSceneShaderData::DataStrideInFloat4s; ArrayIndex++)
 			{
 				for (int32 DataIndex = 0; DataIndex < InstanceSceneData.Num(); DataIndex++)
 				{
@@ -529,7 +529,7 @@ void FCachedRayTracingSceneData::SetupViewUniformBufferFromSceneRenderState(FSce
 			FRHIResourceCreateInfo CreateInfo(TEXT("InstanceSceneDataBuffer"), &InstanceSceneDataSOA);
 			if (InstanceSceneDataSOA.GetResourceDataSize() == 0)
 			{
-				InstanceSceneDataSOA.AddZeroed(FInstanceSceneShaderData::InstanceDataStrideInFloat4s);
+				InstanceSceneDataSOA.AddZeroed(FInstanceSceneShaderData::DataStrideInFloat4s);
 			}
 
 			InstanceSceneDataBufferRHI = RHICreateStructuredBuffer(sizeof(FVector4), InstanceSceneDataSOA.GetResourceDataSize(), BUF_Static | BUF_ShaderResource, CreateInfo);
@@ -790,7 +790,7 @@ void FSceneRenderState::SetupRayTracingScene(int32 LODIndex)
 				TVC_MAX,
 				*View.CachedViewUniformShaderParameters);
 
-			View.CachedViewUniformShaderParameters->InstanceDataSOAStride = CachedRayTracingScene->InstanceSceneDataSOAStride;
+			View.CachedViewUniformShaderParameters->InstanceSceneDataSOAStride = CachedRayTracingScene->InstanceSceneDataSOAStride;
 
 			if (LightSceneRenderState.SkyLight.IsSet())
 			{
