Fix for LWC build error
#rnx

diff --git a/Engine/Plugins/Experimental/GPULightmass/Source/GPULightmass/Private/LightmapRayTracing.h b/Engine/Plugins/Experimental/GPULightmass/Source/GPULightmass/Private/LightmapRayTracing.h
--- a/Engine/Plugins/Experimental/GPULightmass/Source/GPULightmass/Private/LightmapRayTracing.h
+++ b/Engine/Plugins/Experimental/GPULightmass/Source/GPULightmass/Private/LightmapRayTracing.h
@@ -15,8 +15,8 @@
 
 BEGIN_SHADER_PARAMETER_STRUCT(FPathTracingLightGrid, RENDERER_API)
 	SHADER_PARAMETER(uint32, SceneInfiniteLightCount)
-	SHADER_PARAMETER(FVector, SceneLightsBoundMin)
-	SHADER_PARAMETER(FVector, SceneLightsBoundMax)
+	SHADER_PARAMETER(FVector3f, SceneLightsBoundMin)
+	SHADER_PARAMETER(FVector3f, SceneLightsBoundMax)
 	SHADER_PARAMETER_RDG_TEXTURE(Texture2D, LightGrid)
 	SHADER_PARAMETER_RDG_BUFFER_SRV(Buffer<uint>, LightGridData)
 	SHADER_PARAMETER(unsigned, LightGridResolution)
