Fixed build break.

diff --git a/Engine/Plugins/Experimental/GPULightmass/Source/GPULightmass/Private/LightmapRenderer.cpp b/Engine/Plugins/Experimental/GPULightmass/Source/GPULightmass/Private/LightmapRenderer.cpp
--- a/Engine/Plugins/Experimental/GPULightmass/Source/GPULightmass/Private/LightmapRenderer.cpp
+++ b/Engine/Plugins/Experimental/GPULightmass/Source/GPULightmass/Private/LightmapRenderer.cpp
@@ -2611,7 +2611,7 @@ void FLightmapRenderer::Finalize(FRDGBuilder& GraphBuilder)
 							RDG_EVENT_NAME("StationaryLightShadowTracing %dx%d", RayTracingResolution.X, RayTracingResolution.Y),
 							PassParameters,
 							ERDGPassFlags::Compute,
-							[PassParameters, this, RayTracingScene = Scene->RayTracingScene, PipelineState = Scene->RayTracingPipelineState, RayGenerationShader, RayTracingResolution](FRHIRayTracingCommandList& RHICmdListImmediate)
+							[PassParameters, this, RayTracingScene = Scene->RayTracingScene, PipelineState = Scene->RayTracingPipelineState, RayGenerationShader, RayTracingResolution](FRHIRayTracingCommandList& RHICmdList)
 						{
 							FRayTracingShaderBindingsWriter GlobalResources;
 							SetShaderParameters(GlobalResources, RayGenerationShader, *PassParameters);
