fix for build error in GPULightmass
#rnx

diff --git a/Engine/Plugins/Experimental/GPULightmass/Source/GPULightmass/Private/LightmapRenderer.cpp b/Engine/Plugins/Experimental/GPULightmass/Source/GPULightmass/Private/LightmapRenderer.cpp
--- a/Engine/Plugins/Experimental/GPULightmass/Source/GPULightmass/Private/LightmapRenderer.cpp
+++ b/Engine/Plugins/Experimental/GPULightmass/Source/GPULightmass/Private/LightmapRenderer.cpp
@@ -2248,7 +2248,7 @@ void FLightmapRenderer::Finalize(FRDGBuilder& GraphBuilder)
 									SetupPathTracingLightParameters(Scene->LightSceneRenderState, GraphBuilder, PassParameters);
 
 									// TODO: find a way to share IES atlas with path tracer ...
-									PassParameters->IESTexture = GWhiteTexture->TextureRHI;
+									PassParameters->IESTexture = GraphBuilder.RegisterExternalTexture(GSystemTextures.WhiteDummy, TEXT("IESTexture"));
 									PassParameters->IESTextureSampler = TStaticSamplerState<SF_Bilinear, AM_Clamp, AM_Clamp, AM_Clamp>::GetRHI();
 
 									PassParameters->SSProfilesTexture = GetSubsufaceProfileTexture_RT(GraphBuilder.RHICmdList)->GetShaderResourceRHI();
