Moved some duplicated code to new helper function CreateRayTracingSceneWithGeometryInstances

diff --git a/Engine/Plugins/Experimental/GPULightmass/Source/GPULightmass/Private/LightmapRenderer.cpp b/Engine/Plugins/Experimental/GPULightmass/Source/GPULightmass/Private/LightmapRenderer.cpp
--- a/Engine/Plugins/Experimental/GPULightmass/Source/GPULightmass/Private/LightmapRenderer.cpp
+++ b/Engine/Plugins/Experimental/GPULightmass/Source/GPULightmass/Private/LightmapRenderer.cpp
@@ -26,7 +26,6 @@
 #include "RayTracingInstanceBufferUtil.h"
 #include "ScreenPass.h"
 #include "RayTracingDynamicGeometryCollection.h"
-#include "Experimental/Containers/SherwoodHashTable.h"
 
 class FCopyConvergedLightmapTilesCS : public FGlobalShader
 {
@@ -167,57 +166,6 @@ struct FGPUBatchedTileRequests
 	TResourceArray<FGPUTileDescription> BatchedTilesDesc;
 };
 
-static FRayTracingSceneRHIRef CreateRayTracingScene(TArrayView<FRayTracingGeometryInstance> Instances, TArray<uint32>& OutInstancesGeometryIndex)
-{
-	const uint32 NumSceneInstances = Instances.Num();
-
-	FRayTracingSceneInitializer2 Initializer;
-	Initializer.ShaderSlotsPerGeometrySegment = RAY_TRACING_NUM_SHADER_SLOTS;
-	Initializer.PerInstanceGeometries.SetNumUninitialized(NumSceneInstances);
-	Initializer.BaseInstancePrefixSum.SetNumUninitialized(NumSceneInstances);
-	Initializer.SegmentPrefixSum.SetNumUninitialized(NumSceneInstances);
-	Initializer.NumNativeInstances = 0;
-	Initializer.NumTotalSegments = 0;
-
-	Experimental::TSherwoodMap<FRHIRayTracingGeometry*, uint32> UniqueGeometries;
-	OutInstancesGeometryIndex.SetNumUninitialized(NumSceneInstances);
-
-	// Compute geometry segment and instance count prefix sums.
-	// These are later used by GetHitRecordBaseIndex() during resource binding
-	// and by GetBaseInstanceIndex() in shaders to emulate SV_InstanceIndex.
-
-	for (uint32 InstanceIndex = 0; InstanceIndex < NumSceneInstances; ++InstanceIndex)
-	{
-		const FRayTracingGeometryInstance& InstanceDesc = Instances[InstanceIndex];
-
-		checkf(InstanceDesc.GPUTransformsSRV || InstanceDesc.NumTransforms <= uint32(InstanceDesc.Transforms.Num()),
-			TEXT("Expected at most %d ray tracing geometry instance transforms, but got %d."),
-			InstanceDesc.NumTransforms, InstanceDesc.Transforms.Num());
-
-		checkf(InstanceDesc.GeometryRHI, TEXT("Ray tracing instance must have a valid geometry."));
-
-		ensure(InstanceDesc.NumTransforms == uint32(InstanceDesc.Transforms.Num()));
-
-		Initializer.PerInstanceGeometries[InstanceIndex] = InstanceDesc.GeometryRHI;
-
-		// Compute geometry segment count prefix sum to be later used in GetHitRecordBaseIndex()
-		Initializer.SegmentPrefixSum[InstanceIndex] = Initializer.NumTotalSegments;
-		Initializer.NumTotalSegments += InstanceDesc.GeometryRHI->GetNumSegments();
-
-		uint32 GeometryIndex = UniqueGeometries.FindOrAdd(InstanceDesc.GeometryRHI, Initializer.ReferencedGeometries.Num());
-		OutInstancesGeometryIndex[InstanceIndex] = GeometryIndex;
-		if (GeometryIndex == Initializer.ReferencedGeometries.Num())
-		{
-			Initializer.ReferencedGeometries.Add(InstanceDesc.GeometryRHI);
-		}
-
-		Initializer.BaseInstancePrefixSum[InstanceIndex] = Initializer.NumNativeInstances;
-		Initializer.NumNativeInstances += InstanceDesc.NumTransforms;
-	}
-
-	return RHICreateRayTracingScene(MoveTemp(Initializer));
-}
-
 namespace GPULightmass
 {
 
@@ -1056,8 +1004,8 @@ void FSceneRenderState::SetupRayTracingScene(int32 LODIndex)
 		{
 			SCOPED_GPU_MASK(RHICmdList, FRHIGPUMask::All());
 
-			TArray<uint32> InstancesGeometryIndex;
-			RayTracingScene = CreateRayTracingScene(RayTracingGeometryInstances, InstancesGeometryIndex);
+			TArray<uint32> GeometryIndices;
+			RayTracingScene = CreateRayTracingSceneWithGeometryInstances(RayTracingGeometryInstances, RAY_TRACING_NUM_SHADER_SLOTS, 1, GeometryIndices);
 
 			const FRayTracingSceneInitializer2& SceneInitializer = RayTracingScene->GetInitializer();
 
@@ -1078,19 +1026,18 @@ void FSceneRenderState::SetupRayTracingScene(int32 LODIndex)
 			FByteAddressBuffer AccelerationStructureAddressesBuffer;
 			AccelerationStructureAddressesBuffer.Initialize(TEXT("LightmassRayTracingAccelerationStructureAddressesBuffer"), sizeof(FRayTracingAccelerationStructureAddress), BUF_Volatile);
 
+			const uint32 InstanceUploadBufferSize = SceneInitializer.NumNativeInstances * sizeof(FRayTracingInstanceDescriptorInput);
 			FBufferRHIRef InstanceUploadBuffer;
 			FShaderResourceViewRHIRef InstanceUploadSRV;
 			{
-				const uint32 UploadBufferSize = SceneInitializer.NumNativeInstances * sizeof(FRayTracingInstanceDescriptorInput);
-
 				FRHIResourceCreateInfo CreateInfo(TEXT("RayTracingTestBedInstanceUploadBuffer"));
-				InstanceUploadBuffer = RHICreateStructuredBuffer(sizeof(FRayTracingInstanceDescriptorInput), UploadBufferSize, BUF_ShaderResource | BUF_Volatile, CreateInfo);
+				InstanceUploadBuffer = RHICreateStructuredBuffer(sizeof(FRayTracingInstanceDescriptorInput), InstanceUploadBufferSize, BUF_ShaderResource | BUF_Volatile, CreateInfo);
 				InstanceUploadSRV = RHICreateShaderResourceView(InstanceUploadBuffer);
 			}
 
 			{
-				FRayTracingInstanceDescriptorInput* InstanceUploadData = (FRayTracingInstanceDescriptorInput*)RHICmdList.LockBuffer(InstanceUploadBuffer, 0, sizeof(FRayTracingInstanceDescriptorInput), RLM_WriteOnly);
-				FillInstanceUploadBuffer(RayTracingGeometryInstances, InstancesGeometryIndex, RayTracingScene, MakeArrayView(InstanceUploadData, SceneInitializer.NumNativeInstances));
+				FRayTracingInstanceDescriptorInput* InstanceUploadData = (FRayTracingInstanceDescriptorInput*)RHICmdList.LockBuffer(InstanceUploadBuffer, 0, InstanceUploadBufferSize, RLM_WriteOnly);
+				FillRayTracingInstanceUploadBuffer(RayTracingGeometryInstances, GeometryIndices, RayTracingScene, MakeArrayView(InstanceUploadData, SceneInitializer.NumNativeInstances));
 				RHICmdList.UnlockBuffer(InstanceUploadBuffer);
 			}
 
diff --git a/Engine/Source/Runtime/Renderer/Private/RayTracing/RayTracingInstanceBufferUtil.cpp b/Engine/Source/Runtime/Renderer/Private/RayTracing/RayTracingInstanceBufferUtil.cpp
--- a/Engine/Source/Runtime/Renderer/Private/RayTracing/RayTracingInstanceBufferUtil.cpp
+++ b/Engine/Source/Runtime/Renderer/Private/RayTracing/RayTracingInstanceBufferUtil.cpp
@@ -12,9 +12,70 @@
 
 #include "Async/ParallelFor.h"
 
+#include "Experimental/Containers/SherwoodHashTable.h"
+
 #if RHI_RAYTRACING
 
-void FillInstanceUploadBuffer(
+FRayTracingSceneRHIRef CreateRayTracingSceneWithGeometryInstances(
+	TArrayView<FRayTracingGeometryInstance> Instances,
+	uint32 NumShaderSlotsPerGeometrySegment,
+	uint32 NumMissShaderSlots,
+	TArray<uint32>& OutGeometryIndices)
+{
+	const uint32 NumSceneInstances = Instances.Num();
+
+	OutGeometryIndices.SetNumUninitialized(NumSceneInstances);
+
+	FRayTracingSceneInitializer2 Initializer;
+	Initializer.DebugName = FName(TEXT("FRayTracingScene"));
+	Initializer.ShaderSlotsPerGeometrySegment = NumShaderSlotsPerGeometrySegment;
+	Initializer.NumMissShaderSlots = NumMissShaderSlots;
+	Initializer.PerInstanceGeometries.SetNumUninitialized(NumSceneInstances);
+	Initializer.BaseInstancePrefixSum.SetNumUninitialized(NumSceneInstances);
+	Initializer.SegmentPrefixSum.SetNumUninitialized(NumSceneInstances);
+	Initializer.NumNativeInstances = 0;
+	Initializer.NumTotalSegments = 0;
+
+	Experimental::TSherwoodMap<FRHIRayTracingGeometry*, uint32> UniqueGeometries;
+
+	// Compute geometry segment and instance count prefix sums.
+	// These are later used by GetHitRecordBaseIndex() during resource binding
+	// and by GetBaseInstanceIndex() in shaders to emulate SV_InstanceIndex.
+
+	for (uint32 InstanceIndex = 0; InstanceIndex < NumSceneInstances; ++InstanceIndex)
+	{
+		const FRayTracingGeometryInstance& InstanceDesc = Instances[InstanceIndex];
+
+		const bool bGpuInstance = InstanceDesc.GPUTransformsSRV != nullptr;
+		const bool bCpuInstance = !bGpuInstance;
+
+		checkf(bGpuInstance || InstanceDesc.NumTransforms <= uint32(InstanceDesc.Transforms.Num()),
+			TEXT("Expected at most %d ray tracing geometry instance transforms, but got %d."),
+			InstanceDesc.NumTransforms, InstanceDesc.Transforms.Num());
+
+		checkf(InstanceDesc.GeometryRHI, TEXT("Ray tracing instance must have a valid geometry."));
+
+		Initializer.PerInstanceGeometries[InstanceIndex] = InstanceDesc.GeometryRHI;
+
+		// Compute geometry segment count prefix sum to be later used in GetHitRecordBaseIndex()
+		Initializer.SegmentPrefixSum[InstanceIndex] = Initializer.NumTotalSegments;
+		Initializer.NumTotalSegments += InstanceDesc.GeometryRHI->GetNumSegments();
+
+		uint32 GeometryIndex = UniqueGeometries.FindOrAdd(InstanceDesc.GeometryRHI, Initializer.ReferencedGeometries.Num());
+		OutGeometryIndices[InstanceIndex] = GeometryIndex;
+		if (GeometryIndex == Initializer.ReferencedGeometries.Num())
+		{
+			Initializer.ReferencedGeometries.Add(InstanceDesc.GeometryRHI);
+		}
+
+		Initializer.BaseInstancePrefixSum[InstanceIndex] = Initializer.NumNativeInstances;
+		Initializer.NumNativeInstances += InstanceDesc.NumTransforms;
+	}
+
+	return RHICreateRayTracingScene(MoveTemp(Initializer));
+}
+
+void FillRayTracingInstanceUploadBuffer(
 	TConstArrayView<FRayTracingGeometryInstance> Instances,
 	TConstArrayView<uint32> InstancesGeometryIndex,
 	FRayTracingSceneRHIRef RayTracingSceneRHI,
diff --git a/Engine/Source/Runtime/Renderer/Private/RayTracing/RayTracingScene.cpp b/Engine/Source/Runtime/Renderer/Private/RayTracing/RayTracingScene.cpp
--- a/Engine/Source/Runtime/Renderer/Private/RayTracing/RayTracingScene.cpp
+++ b/Engine/Source/Runtime/Renderer/Private/RayTracing/RayTracingScene.cpp
@@ -10,8 +10,6 @@
 #include "RenderGraphBuilder.h"
 #include "RenderGraphUtils.h"
 
-#include "Experimental/Containers/SherwoodHashTable.h"
-
 BEGIN_SHADER_PARAMETER_STRUCT(FBuildInstanceBufferPassParams, )
 	SHADER_PARAMETER_RDG_BUFFER_UAV(RWStructuredBuffer, InstanceBuffer)
 END_SHADER_PARAMETER_STRUCT()
@@ -32,57 +30,8 @@ void FRayTracingScene::Create(FRDGBuilder& GraphBuilder)
 
 	WaitForTasks();
 
-	TArray<uint32> InstancesGeometryIndex;
-
-	{
-		const uint32 NumSceneInstances = Instances.Num();
-
-		FRayTracingSceneInitializer2 SceneInitializer;
-		SceneInitializer.DebugName = FName(TEXT("FRayTracingScene"));
-		SceneInitializer.ShaderSlotsPerGeometrySegment = RAY_TRACING_NUM_SHADER_SLOTS;
-		SceneInitializer.NumMissShaderSlots = RAY_TRACING_NUM_MISS_SHADER_SLOTS;
-		SceneInitializer.PerInstanceGeometries.SetNumUninitialized(NumSceneInstances);
-		SceneInitializer.BaseInstancePrefixSum.SetNumUninitialized(NumSceneInstances);
-		SceneInitializer.SegmentPrefixSum.SetNumUninitialized(NumSceneInstances);
-		SceneInitializer.NumNativeInstances = 0;
-		SceneInitializer.NumTotalSegments = 0;
-
-		Experimental::TSherwoodMap<FRHIRayTracingGeometry*, uint32> UniqueGeometries;
-		InstancesGeometryIndex.SetNumUninitialized(NumSceneInstances);
-
-		// Compute geometry segment and instance count prefix sums.
-		// These are later used by GetHitRecordBaseIndex() during resource binding
-		// and by GetBaseInstanceIndex() in shaders to emulate SV_InstanceIndex.
-
-		for (uint32 InstanceIndex = 0; InstanceIndex < NumSceneInstances; ++InstanceIndex)
-		{
-			const FRayTracingGeometryInstance& InstanceDesc = Instances[InstanceIndex];
-
-			checkf(InstanceDesc.GPUTransformsSRV || InstanceDesc.NumTransforms <= uint32(InstanceDesc.Transforms.Num()),
-				TEXT("Expected at most %d ray tracing geometry instance transforms, but got %d."),
-				InstanceDesc.NumTransforms, InstanceDesc.Transforms.Num());
-
-			checkf(InstanceDesc.GeometryRHI, TEXT("Ray tracing instance must have a valid geometry."));
-
-			SceneInitializer.PerInstanceGeometries[InstanceIndex] = InstanceDesc.GeometryRHI;
-
-			// Compute geometry segment count prefix sum to be later used in GetHitRecordBaseIndex()
-			SceneInitializer.SegmentPrefixSum[InstanceIndex] = SceneInitializer.NumTotalSegments;
-			SceneInitializer.NumTotalSegments += InstanceDesc.GeometryRHI->GetNumSegments();
-
-			uint32 GeometryIndex = UniqueGeometries.FindOrAdd(InstanceDesc.GeometryRHI, SceneInitializer.ReferencedGeometries.Num());
-			InstancesGeometryIndex[InstanceIndex] = GeometryIndex;
-			if (GeometryIndex == SceneInitializer.ReferencedGeometries.Num())
-			{
-				SceneInitializer.ReferencedGeometries.Add(InstanceDesc.GeometryRHI);
-			}
-
-			SceneInitializer.BaseInstancePrefixSum[InstanceIndex] = SceneInitializer.NumNativeInstances;
-			SceneInitializer.NumNativeInstances += InstanceDesc.NumTransforms;
-		}
-
-		RayTracingSceneRHI = RHICreateRayTracingScene(MoveTemp(SceneInitializer));
-	}
+	TArray<uint32> GeometryIndices;
+	RayTracingSceneRHI = CreateRayTracingSceneWithGeometryInstances(Instances, RAY_TRACING_NUM_SHADER_SLOTS, RAY_TRACING_NUM_MISS_SHADER_SLOTS, GeometryIndices);
 
 	const FRayTracingSceneInitializer2& SceneInitializer = RayTracingSceneRHI->GetInitializer();
 
@@ -167,12 +116,12 @@ void FRayTracingScene::Create(FRDGBuilder& GraphBuilder)
 			[InstanceUploadData,
 			NumNativeInstances = SceneInitializer.NumNativeInstances,
 			Instances = MakeArrayView(Instances),
-			InstancesGeometryIndex = MoveTemp(InstancesGeometryIndex),
+			GeometryIndices = MoveTemp(GeometryIndices),
 			RayTracingSceneRHI = RayTracingSceneRHI]()
 		{
 			FTaskTagScope TaskTagScope(ETaskTag::EParallelRenderingThread);
 
-			FillInstanceUploadBuffer(Instances, InstancesGeometryIndex, RayTracingSceneRHI, MakeArrayView(InstanceUploadData, NumNativeInstances));
+			FillRayTracingInstanceUploadBuffer(Instances, GeometryIndices, RayTracingSceneRHI, MakeArrayView(InstanceUploadData, NumNativeInstances));
 		}, TStatId(), nullptr, ENamedThreads::AnyThread);
 
 		FBuildInstanceBufferPassParams* PassParams = GraphBuilder.AllocParameters<FBuildInstanceBufferPassParams>();
diff --git a/Engine/Source/Runtime/Renderer/Private/Tests/RayTracingTestbed.cpp b/Engine/Source/Runtime/Renderer/Private/Tests/RayTracingTestbed.cpp
--- a/Engine/Source/Runtime/Renderer/Private/Tests/RayTracingTestbed.cpp
+++ b/Engine/Source/Runtime/Renderer/Private/Tests/RayTracingTestbed.cpp
@@ -115,24 +115,8 @@ bool RunRayTracingTestbed_RenderThread(const FString& Parameters)
 	Instances[0].NumTransforms = NumTransforms;
 	Instances[0].Transforms = MakeArrayView(&FMatrix::Identity, 1);
 
-	uint32 InstancesGeometryIndex[NumInstances] = {};
-	InstancesGeometryIndex[0] = 0;
-
-	FRayTracingSceneRHIRef Scene;
-
-	{
-		FRayTracingSceneInitializer2 Initializer;
-		Initializer.DebugName = FName(TEXT("FRayTracingScene"));
-		Initializer.ShaderSlotsPerGeometrySegment = RAY_TRACING_NUM_SHADER_SLOTS;
-		Initializer.PerInstanceGeometries.Add(Geometry);
-		Initializer.BaseInstancePrefixSum.Add(0);
-		Initializer.SegmentPrefixSum.Add(0);
-		Initializer.ReferencedGeometries.Add(Geometry);
-		Initializer.NumNativeInstances = NumInstances * NumTransforms;
-		Initializer.NumTotalSegments = Geometry->GetNumSegments();
-
-		Scene = RHICreateRayTracingScene(MoveTemp(Initializer));
-	}
+	TArray<uint32> GeometryIndices;
+	FRayTracingSceneRHIRef Scene = CreateRayTracingSceneWithGeometryInstances(Instances, RAY_TRACING_NUM_SHADER_SLOTS, 1, GeometryIndices);
 
 	const FRayTracingSceneInitializer2& SceneInitializer = Scene->GetInitializer();
 
@@ -151,26 +135,26 @@ bool RunRayTracingTestbed_RenderThread(const FString& Parameters)
 		ScratchBufferCreateInfo);
 
 	FRWBufferStructured InstanceBuffer;
-	InstanceBuffer.Initialize(TEXT("RayTracingTestBedInstanceBuffer"), GRHIRayTracingInstanceDescriptorSize, 1);
+	InstanceBuffer.Initialize(TEXT("RayTracingTestBedInstanceBuffer"), GRHIRayTracingInstanceDescriptorSize, SceneInitializer.NumNativeInstances);
 
 	FByteAddressBuffer AccelerationStructureAddressesBuffer;
 	AccelerationStructureAddressesBuffer.Initialize(TEXT("RayTracingTestBedAccelerationStructureAddressesBuffer"), sizeof(FRayTracingAccelerationStructureAddress), BUF_Volatile);
 
+	const uint32 InstanceUploadBufferSize = SceneInitializer.NumNativeInstances * sizeof(FRayTracingInstanceDescriptorInput);
 	FBufferRHIRef InstanceUploadBuffer;
 	FShaderResourceViewRHIRef InstanceUploadSRV;
 	{
-		const uint32 UploadBufferSize = 1 * sizeof(FRayTracingInstanceDescriptorInput);
 
 		FRHIResourceCreateInfo CreateInfo(TEXT("RayTracingTestBedInstanceUploadBuffer"));
-		InstanceUploadBuffer = RHICreateStructuredBuffer(sizeof(FRayTracingInstanceDescriptorInput), UploadBufferSize, BUF_ShaderResource | BUF_Volatile, CreateInfo);
+		InstanceUploadBuffer = RHICreateStructuredBuffer(sizeof(FRayTracingInstanceDescriptorInput), InstanceUploadBufferSize, BUF_ShaderResource | BUF_Volatile, CreateInfo);
 		InstanceUploadSRV = RHICreateShaderResourceView(InstanceUploadBuffer);
 	}
 
 	FRHICommandListImmediate& RHICmdList = FRHICommandListExecutor::GetImmediateCommandList();
 
 	{
-		FRayTracingInstanceDescriptorInput* InstanceUploadData = (FRayTracingInstanceDescriptorInput*)RHICmdList.LockBuffer(InstanceUploadBuffer, 0, sizeof(FRayTracingInstanceDescriptorInput), RLM_WriteOnly);
-		FillInstanceUploadBuffer(Instances, InstancesGeometryIndex, Scene, MakeArrayView(InstanceUploadData, 1));
+		FRayTracingInstanceDescriptorInput* InstanceUploadData = (FRayTracingInstanceDescriptorInput*)RHICmdList.LockBuffer(InstanceUploadBuffer, 0, InstanceUploadBufferSize, RLM_WriteOnly);
+		FillRayTracingInstanceUploadBuffer(Instances, GeometryIndices, Scene, MakeArrayView(InstanceUploadData, SceneInitializer.NumNativeInstances));
 		RHICmdList.UnlockBuffer(InstanceUploadBuffer);
 	}
 
diff --git a/Engine/Source/Runtime/Renderer/Public/RayTracingInstanceBufferUtil.h b/Engine/Source/Runtime/Renderer/Public/RayTracingInstanceBufferUtil.h
--- a/Engine/Source/Runtime/Renderer/Public/RayTracingInstanceBufferUtil.h
+++ b/Engine/Source/Runtime/Renderer/Public/RayTracingInstanceBufferUtil.h
@@ -17,7 +17,15 @@ struct FRayTracingInstanceDescriptorInput
 	uint32 InstanceContributionToHitGroupIndex;
 };
 
-RENDERER_API void FillInstanceUploadBuffer(
+// Helper function to create FRayTracingSceneRHI using array of high level instances
+// Also outputs data required to build the instance buffer
+RENDERER_API FRayTracingSceneRHIRef CreateRayTracingSceneWithGeometryInstances(
+	TArrayView<FRayTracingGeometryInstance> Instances,
+	uint32 NumShaderSlotsPerGeometrySegment,
+	uint32 NumMissShaderSlots,
+	TArray<uint32>& OutGeometryIndices);
+
+RENDERER_API void FillRayTracingInstanceUploadBuffer(
 	TConstArrayView<FRayTracingGeometryInstance> Instances,
 	TConstArrayView<uint32> InstancesGeometryIndex,
 	FRayTracingSceneRHIRef RayTracingSceneRHI,
