Fix potentially dangerous uses of the ternary operator in the path tracer

The HLSL spec says that the ternary operator may evaluate both sides of the expression which could be dangerous when the ternary operator is guarding array access.

diff --git a/Engine/Plugins/Experimental/GPULightmass/Shaders/Private/LightmapPathTracing.usf b/Engine/Plugins/Experimental/GPULightmass/Shaders/Private/LightmapPathTracing.usf
--- a/Engine/Plugins/Experimental/GPULightmass/Shaders/Private/LightmapPathTracing.usf
+++ b/Engine/Plugins/Experimental/GPULightmass/Shaders/Private/LightmapPathTracing.usf
@@ -677,7 +677,11 @@ void PathTracingKernel(
 				float3 LightContrib = PathThroughput * LightResult.Radiance;
 				if (bUseMIS)
 				{
-					float PreviousCdfValue = Index > 0 ? LightPickingCdf[Index - 1] : 0.0;
+					float PreviousCdfValue = 0.0;
+					BRANCH if (Index > 0)
+					{
+						PreviousCdfValue = LightPickingCdf[Index - 1];
+					}
 					float LightPickPdf = (LightPickingCdf[Index] - PreviousCdfValue) / LightPickingCdfSum;
 
 					LightContrib *= MISWeightRobust(MaterialSample.Pdf, LightResult.Pdf * LightPickPdf);
diff --git a/Engine/Shaders/Private/PathTracing/Light/PathTracingLightSampling.ush b/Engine/Shaders/Private/PathTracing/Light/PathTracingLightSampling.ush
--- a/Engine/Shaders/Private/PathTracing/Light/PathTracingLightSampling.ush
+++ b/Engine/Shaders/Private/PathTracing/Light/PathTracingLightSampling.ush
@@ -139,8 +139,11 @@ void SelectLight(float RandSample, int NumLights, inout float LightPickingCdf[RA
 			Count -= Step + 1;
 		}
 	}
-
-	LightPickPdf = LightPickingCdf[LightId] - ((LightId > 0) ? LightPickingCdf[LightId - 1] : 0.0);
+	LightPickPdf = LightPickingCdf[LightId];
+	BRANCH if (LightId > 0)
+	{
+		LightPickPdf -= LightPickingCdf[LightId - 1];
+	}
 #endif
 #endif
 }
diff --git a/Engine/Shaders/Private/PathTracing/PathTracingCore.ush b/Engine/Shaders/Private/PathTracing/PathTracingCore.ush
--- a/Engine/Shaders/Private/PathTracing/PathTracingCore.ush
+++ b/Engine/Shaders/Private/PathTracing/PathTracingCore.ush
@@ -936,7 +936,11 @@ bool PathTracingKernel(inout FPathState PathState, int Bounce)
 			float3 LightContrib = PathState.PathThroughput * LightResult.Radiance;
 			if (bUseMIS)
 			{
-				float PreviousCdfValue = Index > 0 ? LightPickingCdf[Index - 1] : 0.0;
+				float PreviousCdfValue = 0.0;
+				BRANCH if (Index > 0)
+				{
+					PreviousCdfValue = LightPickingCdf[Index - 1];
+				}
 				float LightPickPdf = (LightPickingCdf[Index] - PreviousCdfValue) / LightPickingCdfSum;
 
 				LightContrib *= MISWeightRobust(MaterialSample.Pdf, LightResult.Pdf * LightPickPdf);
