Fix compile issue introduced in CL-17955874
LocalToWorld renamed to LocalToRelativeWorld

diff --git a/Engine/Plugins/Experimental/GPULightmass/Source/GPULightmass/Private/LightmapRenderer.cpp b/Engine/Plugins/Experimental/GPULightmass/Source/GPULightmass/Private/LightmapRenderer.cpp
--- a/Engine/Plugins/Experimental/GPULightmass/Source/GPULightmass/Private/LightmapRenderer.cpp
+++ b/Engine/Plugins/Experimental/GPULightmass/Source/GPULightmass/Private/LightmapRenderer.cpp
@@ -328,8 +328,8 @@ void FCachedRayTracingSceneData::SetupViewUniformBufferFromSceneRenderState(FSce
 					INVALID_LAST_UPDATE_FRAME,
 					0, /* Custom Data Count */
 					0.0f, /* Random ID */
-					PrimitiveUniformShaderParameters.LocalToWorld,
-					PrimitiveUniformShaderParameters.LocalToWorld // TODO: Temporary PrevVelocityHack
+					PrimitiveUniformShaderParameters.LocalToRelativeWorld,
+					PrimitiveUniformShaderParameters.LocalToRelativeWorld // TODO: Temporary PrevVelocityHack
 				);
 
 				for (int32 LODIndex = 0; LODIndex < Instance.LODLightmapRenderStates.Num(); LODIndex++)
@@ -471,8 +471,8 @@ void FCachedRayTracingSceneData::SetupViewUniformBufferFromSceneRenderState(FSce
 				INVALID_LAST_UPDATE_FRAME,
 				0, /* Custom Data Count */
 				0.0f, /* Random ID */
-				PrimitiveUniformShaderParameters.LocalToWorld,
-				PrimitiveUniformShaderParameters.LocalToWorld // TODO: Temporary PrevVelocityHack
+				PrimitiveUniformShaderParameters.LocalToRelativeWorld,
+				PrimitiveUniformShaderParameters.LocalToRelativeWorld // TODO: Temporary PrevVelocityHack
 			);
 
 			PrimitiveSceneData.Add(FPrimitiveSceneShaderData(PrimitiveUniformShaderParameters));
