GPULM Fixes
* Fix wrong usage of IsInRenderingThread() vs IsInParallelRenderingThread()
* Fix crash with ISMs when traditional lightmaps are used and there is no stationary light in the scene
* Fix black output from the denoiser due to LWC types not converted
* Fix GPULM RT scene initialization
* Fix wrong IES orientation with point lights

diff --git a/Engine/Plugins/Experimental/GPULightmass/Source/GPULightmass/Private/GPULightmass.cpp b/Engine/Plugins/Experimental/GPULightmass/Source/GPULightmass/Private/GPULightmass.cpp
--- a/Engine/Plugins/Experimental/GPULightmass/Source/GPULightmass/Private/GPULightmass.cpp
+++ b/Engine/Plugins/Experimental/GPULightmass/Source/GPULightmass/Private/GPULightmass.cpp
@@ -281,7 +281,7 @@ const FLightComponentMapBuildData* FGPULightmass::GetLightComponentMapBuildData(
 
 const FPrecomputedVolumetricLightmap* FGPULightmass::GetPrecomputedVolumetricLightmap()
 {
-	check(IsInRenderingThread() || IsInParallelRenderingThread());
+	check(IsInParallelRenderingThread());
 
 	return Scene.RenderState.VolumetricLightmapRenderer->GetPrecomputedVolumetricLightmapForPreview();
 }
diff --git a/Engine/Plugins/Experimental/GPULightmass/Source/GPULightmass/Private/LightmapDenoising.h b/Engine/Plugins/Experimental/GPULightmass/Source/GPULightmass/Private/LightmapDenoising.h
--- a/Engine/Plugins/Experimental/GPULightmass/Source/GPULightmass/Private/LightmapDenoising.h
+++ b/Engine/Plugins/Experimental/GPULightmass/Source/GPULightmass/Private/LightmapDenoising.h
@@ -24,8 +24,8 @@ struct FDenoiserFilterSet
 #endif
 
 	FIntPoint Size;
-	TArray<FVector> InputBuffer;
-	TArray<FVector> OutputBuffer;
+	TArray<FVector3f> InputBuffer;
+	TArray<FVector3f> OutputBuffer;
 
 	FDenoiserFilterSet(FDenoiserContext& Context, FIntPoint NewSize, bool bSHDenoiser = false);
 
diff --git a/Engine/Plugins/Experimental/GPULightmass/Source/GPULightmass/Private/LightmapRenderer.cpp b/Engine/Plugins/Experimental/GPULightmass/Source/GPULightmass/Private/LightmapRenderer.cpp
--- a/Engine/Plugins/Experimental/GPULightmass/Source/GPULightmass/Private/LightmapRenderer.cpp
+++ b/Engine/Plugins/Experimental/GPULightmass/Source/GPULightmass/Private/LightmapRenderer.cpp
@@ -1011,10 +1011,10 @@ void FSceneRenderState::SetupRayTracingScene(int32 LODIndex)
 				ScratchBufferCreateInfo);
 
 			FRWBufferStructured InstanceBuffer;
-			InstanceBuffer.Initialize(TEXT("LightmassRayTracingInstanceBuffer"), GRHIRayTracingInstanceDescriptorSize, 1);
+			InstanceBuffer.Initialize(TEXT("LightmassRayTracingInstanceBuffer"), GRHIRayTracingInstanceDescriptorSize, SceneInitializer.NumNativeInstances);
 
 			FByteAddressBuffer AccelerationStructureAddressesBuffer;
-			AccelerationStructureAddressesBuffer.Initialize(TEXT("LightmassRayTracingAccelerationStructureAddressesBuffer"), sizeof(FRayTracingAccelerationStructureAddress), BUF_Volatile);
+			AccelerationStructureAddressesBuffer.Initialize(TEXT("LightmassRayTracingAccelerationStructureAddressesBuffer"), SceneInitializer.ReferencedGeometries.Num() * sizeof(FRayTracingAccelerationStructureAddress), BUF_Volatile);
 
 			const uint32 InstanceUploadBufferSize = SceneInitializer.NumNativeInstances * sizeof(FRayTracingInstanceDescriptorInput);
 			FBufferRHIRef InstanceUploadBuffer;
@@ -1050,10 +1050,10 @@ void FSceneRenderState::SetupRayTracingScene(int32 LODIndex)
 				RHICmdList.UnlockBuffer(InstanceUploadBuffer);
 			}
 
-			RHICmdList.EnqueueLambda([&AccelerationStructureAddressesBuffer, &SceneInitializer](FRHICommandListImmediate& RHICmdList)
+			RHICmdList.EnqueueLambda([BufferRHIRef = AccelerationStructureAddressesBuffer.Buffer, &SceneInitializer](FRHICommandListImmediate& RHICmdList)
 				{
 					FRayTracingAccelerationStructureAddress* AddressesPtr = (FRayTracingAccelerationStructureAddress*)RHICmdList.LockBuffer(
-						AccelerationStructureAddressesBuffer.Buffer,
+						BufferRHIRef,
 						0,
 						SceneInitializer.ReferencedGeometries.Num() * sizeof(FRayTracingAccelerationStructureAddress), RLM_WriteOnly);
 
@@ -1063,7 +1063,7 @@ void FSceneRenderState::SetupRayTracingScene(int32 LODIndex)
 						AddressesPtr[GeometryIndex] = SceneInitializer.ReferencedGeometries[GeometryIndex]->GetAccelerationStructureAddress(RHICmdList.GetGPUMask().ToIndex());
 					}
 
-					RHICmdList.UnlockBuffer(AccelerationStructureAddressesBuffer.Buffer);
+					RHICmdList.UnlockBuffer(BufferRHIRef);
 				});
 
 			BuildRayTracingInstanceBuffer(
diff --git a/Engine/Plugins/Experimental/GPULightmass/Source/GPULightmass/Private/PathTracingLightParameters.inl b/Engine/Plugins/Experimental/GPULightmass/Source/GPULightmass/Private/PathTracingLightParameters.inl
--- a/Engine/Plugins/Experimental/GPULightmass/Source/GPULightmass/Private/PathTracingLightParameters.inl
+++ b/Engine/Plugins/Experimental/GPULightmass/Source/GPULightmass/Private/PathTracingLightParameters.inl
@@ -74,7 +74,7 @@ void SetupPathTracingLightParameters(
 
 		DestLight.Position = Light.Position;
 		DestLight.Color = FVector(Light.Color);
-		DestLight.Normal = FVector(1, 0, 0);
+		DestLight.Normal = Light.Direction;
 		DestLight.dPdu = FVector::CrossProduct(Light.Tangent, Light.Direction);
 		DestLight.dPdv = Light.Tangent;
 
diff --git a/Engine/Plugins/Experimental/GPULightmass/Source/GPULightmass/Private/Scene/Lights.cpp b/Engine/Plugins/Experimental/GPULightmass/Source/GPULightmass/Private/Scene/Lights.cpp
--- a/Engine/Plugins/Experimental/GPULightmass/Source/GPULightmass/Private/Scene/Lights.cpp
+++ b/Engine/Plugins/Experimental/GPULightmass/Source/GPULightmass/Private/Scene/Lights.cpp
@@ -244,6 +244,8 @@ FLightShaderParameters FPointLightRenderState::GetLightShaderParameters() const
 	FLightShaderParameters LightParameters;
 
 	LightParameters.Position = Position;
+	LightParameters.Direction = -Direction;
+	LightParameters.Tangent = Tangent;
 	LightParameters.InvRadius = 1.0f / AttenuationRadius;
 	LightParameters.Color = FVector(Color);
 	LightParameters.SourceRadius = SourceRadius;
diff --git a/Engine/Plugins/Experimental/GPULightmass/Source/GPULightmass/Private/Scene/Lights.h b/Engine/Plugins/Experimental/GPULightmass/Source/GPULightmass/Private/Scene/Lights.h
--- a/Engine/Plugins/Experimental/GPULightmass/Source/GPULightmass/Private/Scene/Lights.h
+++ b/Engine/Plugins/Experimental/GPULightmass/Source/GPULightmass/Private/Scene/Lights.h
@@ -15,19 +15,19 @@ class FSkyLightImportanceSamplingData;
 
 struct FLightShaderConstants
 {
-	FVector Position;
-	float InvRadius;
-	FVector Color;
-	float FalloffExponent;
-	FVector   Direction;
-	float     SpecularScale;
-	FVector   Tangent;
-	float     SourceRadius;
-	FVector2D   SpotAngles;
-	float      SoftSourceRadius;
-	float      SourceLength;
-	float      RectLightBarnCosAngle;
-	float      RectLightBarnLength;
+	FVector3f	Position;
+	float		InvRadius;
+	FVector3f	Color;
+	float		FalloffExponent;
+	FVector3f	Direction;
+	float		SpecularScale;
+	FVector3f	Tangent;
+	float		SourceRadius;
+	FVector2f	SpotAngles;
+	float		SoftSourceRadius;
+	float		SourceLength;
+	float		RectLightBarnCosAngle;
+	float		RectLightBarnLength;
 
 	FLightShaderConstants() = default;
 
diff --git a/Engine/Plugins/Experimental/GPULightmass/Source/GPULightmass/Private/Scene/Scene.cpp b/Engine/Plugins/Experimental/GPULightmass/Source/GPULightmass/Private/Scene/Scene.cpp
--- a/Engine/Plugins/Experimental/GPULightmass/Source/GPULightmass/Private/Scene/Scene.cpp
+++ b/Engine/Plugins/Experimental/GPULightmass/Source/GPULightmass/Private/Scene/Scene.cpp
@@ -2640,7 +2640,15 @@ void FScene::ApplyFinishedLightmapsToWorld()
 							Registry, InstanceGroup.ComponentUObject->LODData[LODIndex].MapBuildDataId, InstanceGroup.ComponentUObject->Bounds, PaddingType, LMF_Streamed);
 
 						MeshBuildData.LightMap = NewLightMap;
-
+						
+						for (auto It = InstancedShadowMapData.CreateIterator(); It; It++)
+						{
+							if (It->Num() == 0)
+							{
+								It.RemoveCurrent();
+							}
+						}
+						
 						if (InstancedShadowMapData.Num() > 0 && !bUseVirtualTextures)
 						{
 							TRefCountPtr<FShadowMap2D> NewShadowMap = FShadowMap2D::AllocateInstancedShadowMap(Registry, InstanceGroup.ComponentUObject,
