GPULM: Fix FInstanceSceneShaderData construction

diff --git a/Engine/Plugins/Experimental/GPULightmass/Source/GPULightmass/Private/LightmapRenderer.cpp b/Engine/Plugins/Experimental/GPULightmass/Source/GPULightmass/Private/LightmapRenderer.cpp
--- a/Engine/Plugins/Experimental/GPULightmass/Source/GPULightmass/Private/LightmapRenderer.cpp
+++ b/Engine/Plugins/Experimental/GPULightmass/Source/GPULightmass/Private/LightmapRenderer.cpp
@@ -319,7 +319,7 @@ void FCachedRayTracingSceneData::SetupViewUniformBufferFromSceneRenderState(FSce
 
 				InstanceDataOriginalOffsets[InstanceIndex] = InstanceIndex;
 
-				FInstanceSceneShaderData& SceneData = InstanceSceneData.Emplace_GetRef();
+				FInstanceSceneShaderData& SceneData = InstanceSceneData[InstanceIndex];
 				SceneData.BuildInternal
 				(
 					InstanceIndex, /* Primitive Id */
