GPULM: Properly populate InstancePayloadData for ISMs to fix multiple ISM related issues

diff --git a/Engine/Plugins/Experimental/GPULightmass/Source/GPULightmass/Private/LightmapRenderer.cpp b/Engine/Plugins/Experimental/GPULightmass/Source/GPULightmass/Private/LightmapRenderer.cpp
--- a/Engine/Plugins/Experimental/GPULightmass/Source/GPULightmass/Private/LightmapRenderer.cpp
+++ b/Engine/Plugins/Experimental/GPULightmass/Source/GPULightmass/Private/LightmapRenderer.cpp
@@ -251,7 +251,7 @@ void FCachedRayTracingSceneData::SetupViewUniformBufferFromSceneRenderState(FSce
 	TResourceArray<FPrimitiveSceneShaderData>	PrimitiveSceneData;
 	TResourceArray<FLightmapSceneShaderData>	LightmapSceneData;
 	TResourceArray<FInstanceSceneShaderData>	InstanceSceneData;
-	TResourceArray<FVector4f>					InstancePayloadData; // TODO: Populate this
+	TResourceArray<FVector4f>					InstancePayloadData;
 
 	PrimitiveSceneData.AddZeroed(Scene.StaticMeshInstanceRenderStates.Elements.Num());
 	InstanceSceneData.AddZeroed(Scene.StaticMeshInstanceRenderStates.Elements.Num());
@@ -320,7 +320,7 @@ void FCachedRayTracingSceneData::SetupViewUniformBufferFromSceneRenderState(FSce
 				InstanceDataOriginalOffsets[InstanceIndex] = InstanceIndex;
 
 				FInstanceSceneShaderData& SceneData = InstanceSceneData[InstanceIndex];
-				SceneData.BuildInternal
+				SceneData.Build
 				(
 					InstanceIndex, /* Primitive Id */
 					0, /* Relative Instance Id */
@@ -328,8 +328,9 @@ void FCachedRayTracingSceneData::SetupViewUniformBufferFromSceneRenderState(FSce
 					INVALID_LAST_UPDATE_FRAME,
 					0, /* Custom Data Count */
 					0.0f, /* Random ID */
+					FRenderTransform::Identity,
 					PrimitiveUniformShaderParameters.LocalToRelativeWorld,
-					PrimitiveUniformShaderParameters.LocalToRelativeWorld // TODO: Temporary PrevVelocityHack
+					PrimitiveUniformShaderParameters.LocalToRelativeWorld
 				);
 
 				for (int32 LODIndex = 0; LODIndex < Instance.LODLightmapRenderStates.Num(); LODIndex++)
@@ -371,8 +372,8 @@ void FCachedRayTracingSceneData::SetupViewUniformBufferFromSceneRenderState(FSce
 					.LightmapDataIndex(LightmapSceneDataStartOffsets[PrimitiveId])
 					.InstanceSceneDataOffset(InstanceSceneData.Num())
 					.NumInstanceSceneDataEntries(NumInstancesThisGroup)
-					.InstancePayloadDataOffset(INDEX_NONE) // InstancePayloadData.Num()) // TODO: Populate this
-					.InstancePayloadDataStride(0) // TODO: Populate this
+					.InstancePayloadDataOffset(InstancePayloadData.Num())
+					.InstancePayloadDataStride(1)
 					.CastContactShadow(true)
 					.CastShadow(true)
 				.Build();
@@ -404,17 +405,21 @@ void FCachedRayTracingSceneData::SetupViewUniformBufferFromSceneRenderState(FSce
 				InstanceGroup.InstancedRenderData->PerInstanceRenderData->InstanceBuffer.GetInstanceTransform(InstanceIdx, Instance.LocalToPrimitive);
 
 				FInstanceSceneShaderData& SceneData = InstanceSceneData.Emplace_GetRef();
-				SceneData.BuildInternal
+				SceneData.Build
 				(
 					PrimitiveId,
-					0, /* Relative Instance Id */
-					0, /* Payload Data Flags */
+					InstanceIdx, /* Relative Instance Id */
+					INSTANCE_SCENE_DATA_FLAG_HAS_LIGHTSHADOW_UV_BIAS, /* Payload Data Flags */
 					INVALID_LAST_UPDATE_FRAME,
 					0, /* Custom Data Count */
 					0.0f, /* Random ID */
+					Instance.LocalToPrimitive,
 					InstanceGroup.LocalToWorld,
-					InstanceGroup.LocalToWorld // TODO: Temporary PrevVelocityHack
+					InstanceGroup.LocalToWorld
 				);
+
+				FVector4f& InstanceLightShadowUVBias = InstancePayloadData.Emplace_GetRef();
+				InstanceGroup.InstancedRenderData->PerInstanceRenderData->InstanceBuffer.GetInstanceLightMapData(InstanceIdx, InstanceLightShadowUVBias);
 			}
 
 			PrimitiveSceneData.Add(FPrimitiveSceneShaderData(PrimitiveUniformShaderParameters));
@@ -534,6 +539,19 @@ void FCachedRayTracingSceneData::SetupViewUniformBufferFromSceneRenderState(FSce
 			InstanceSceneDataSOAStride = InstanceSceneData.Num();
 		}
 
+		{
+			TRACE_CPUPROFILER_EVENT_SCOPE(InstancePayloadData);
+
+			FRHIResourceCreateInfo CreateInfo(TEXT("InstancePayloadDataBuffer"), &InstancePayloadData);
+			if (InstancePayloadData.GetResourceDataSize() == 0)
+			{
+				InstancePayloadData.Add(FVector4f());
+			}
+
+			InstancePayloadDataBufferRHI = RHICreateStructuredBuffer(sizeof(FVector4f), InstancePayloadData.GetResourceDataSize(), BUF_Static | BUF_ShaderResource, CreateInfo);
+			InstancePayloadDataBufferSRV = RHICreateShaderResourceView(InstancePayloadDataBufferRHI);
+		}
+
 		FViewUniformShaderParameters ViewUniformBufferParameters;
 		CachedViewUniformBuffer = TUniformBufferRef<FViewUniformShaderParameters>::CreateUniformBufferImmediate(ViewUniformBufferParameters, UniformBuffer_MultiFrame, EUniformBufferValidation::None);
 	}
diff --git a/Engine/Plugins/Experimental/GPULightmass/Source/GPULightmass/Private/Scene/InstancedStaticMesh.h b/Engine/Plugins/Experimental/GPULightmass/Source/GPULightmass/Private/Scene/InstancedStaticMesh.h
--- a/Engine/Plugins/Experimental/GPULightmass/Source/GPULightmass/Private/Scene/InstancedStaticMesh.h
+++ b/Engine/Plugins/Experimental/GPULightmass/Source/GPULightmass/Private/Scene/InstancedStaticMesh.h
@@ -15,9 +15,6 @@ public:
 	UInstancedStaticMeshComponent* ComponentUObject;
 	FStaticMeshRenderData* RenderData;
 	TUniquePtr<FInstancedStaticMeshRenderData> InstancedRenderData;
-	FBufferRHIRef InstanceOriginBuffer;
-	FBufferRHIRef InstanceTransformBuffer;
-	FBufferRHIRef InstanceLightmapBuffer;
 
 	TArray<FIntPoint> LODPerInstanceLightmapSize;
 
diff --git a/Engine/Plugins/Experimental/GPULightmass/Source/GPULightmass/Private/Scene/Scene.cpp b/Engine/Plugins/Experimental/GPULightmass/Source/GPULightmass/Private/Scene/Scene.cpp
--- a/Engine/Plugins/Experimental/GPULightmass/Source/GPULightmass/Private/Scene/Scene.cpp
+++ b/Engine/Plugins/Experimental/GPULightmass/Source/GPULightmass/Private/Scene/Scene.cpp
@@ -1050,10 +1050,6 @@ void FScene::AddGeometryInstanceFromComponent(UInstancedStaticMeshComponent* InC
 		](FRHICommandListImmediate&) mutable
 	{
 
-		InstanceRenderState.InstanceOriginBuffer = InstanceRenderState.InstancedRenderData->PerInstanceRenderData->InstanceBuffer.GetInstanceOriginBuffer();
-		InstanceRenderState.InstanceTransformBuffer = InstanceRenderState.InstancedRenderData->PerInstanceRenderData->InstanceBuffer.GetInstanceTransformBuffer();
-		InstanceRenderState.InstanceLightmapBuffer = InstanceRenderState.InstancedRenderData->PerInstanceRenderData->InstanceBuffer.GetInstanceLightmapBuffer();
-
 		FInstanceGroupRenderStateRef InstanceRenderStateRef = RenderState.InstanceGroupRenderStates.Emplace(MoveTemp(InstanceRenderState));
 
 		for (int32 LODIndex = 0; LODIndex < InstanceLightmapRenderStateInitializers.Num(); LODIndex++)
diff --git a/Engine/Plugins/Experimental/GPULightmass/Source/GPULightmass/Private/Scene/Scene.h b/Engine/Plugins/Experimental/GPULightmass/Source/GPULightmass/Private/Scene/Scene.h
--- a/Engine/Plugins/Experimental/GPULightmass/Source/GPULightmass/Private/Scene/Scene.h
+++ b/Engine/Plugins/Experimental/GPULightmass/Source/GPULightmass/Private/Scene/Scene.h
@@ -73,7 +73,6 @@ struct FCachedRayTracingSceneData
 
 	FBufferRHIRef InstancePayloadDataBufferRHI;
 	FShaderResourceViewRHIRef InstancePayloadDataBufferSRV;
-	//uint32 InstancePayloadDataSOAStride;
 
 	FBufferRHIRef InstanceIdsIdentityBufferRHI;
 	FShaderResourceViewRHIRef InstanceIdsIdentityBufferSRV;
