GPULM: Add ReferenceView to its ViewFamily to resolve a crash

diff --git a/Engine/Plugins/Experimental/GPULightmass/Source/GPULightmass/Private/LightmapRenderer.cpp b/Engine/Plugins/Experimental/GPULightmass/Source/GPULightmass/Private/LightmapRenderer.cpp
--- a/Engine/Plugins/Experimental/GPULightmass/Source/GPULightmass/Private/LightmapRenderer.cpp
+++ b/Engine/Plugins/Experimental/GPULightmass/Source/GPULightmass/Private/LightmapRenderer.cpp
@@ -784,6 +784,7 @@ void FSceneRenderState::SetupRayTracingScene(int32 LODIndex)
 
 	ReferenceView = MakeShared<FViewInfo>(ViewInitOptions);
 	FViewInfo& View = *ReferenceView;
+	ViewFamily.Views.Add(ReferenceView.Get());
 	View.ViewRect = View.UnscaledViewRect;
 
 	{
