Fixed crash in GPU lightmass when baking an empty level

diff --git a/Engine/Plugins/Experimental/GPULightmass/Source/GPULightmass/Private/LightmapRenderer.cpp b/Engine/Plugins/Experimental/GPULightmass/Source/GPULightmass/Private/LightmapRenderer.cpp
--- a/Engine/Plugins/Experimental/GPULightmass/Source/GPULightmass/Private/LightmapRenderer.cpp
+++ b/Engine/Plugins/Experimental/GPULightmass/Source/GPULightmass/Private/LightmapRenderer.cpp
@@ -724,7 +724,7 @@ void FCachedRayTracingSceneData::SetupFromSceneRenderState(FSceneRenderState& Sc
 #endif // RHI_RAYTRACING
 }
 
-void FSceneRenderState::SetupRayTracingScene(int32 LODIndex)
+bool FSceneRenderState::SetupRayTracingScene(int32 LODIndex)
 {
 	TRACE_CPUPROFILER_EVENT_SCOPE(SetupRayTracingScene);
 
@@ -745,6 +745,11 @@ void FSceneRenderState::SetupRayTracingScene(int32 LODIndex)
 		CalculateDistributionPrefixSumForAllLightmaps();
 	}
 
+	if (CachedRayTracingScene->RayTracingGeometryInstancesPerLOD[LODIndex].Num() == 0)
+	{
+		return false;
+	}
+
 #if 0 // Debug: verify cached ray tracing scene has up-to-date shader bindings
 	TUniquePtr<FCachedRayTracingSceneData> VerificationRayTracingScene = MakeUnique<FCachedRayTracingSceneData>();
 	VerificationRayTracingScene->CachedViewUniformBuffer = CachedRayTracingScene->CachedViewUniformBuffer;
@@ -1217,6 +1222,8 @@ void FSceneRenderState::SetupRayTracingScene(int32 LODIndex)
 		}
 	}
 #endif
+
+	return true;
 }
 
 void FSceneRenderState::DestroyRayTracingScene()
diff --git a/Engine/Plugins/Experimental/GPULightmass/Source/GPULightmass/Private/Scene/Scene.h b/Engine/Plugins/Experimental/GPULightmass/Source/GPULightmass/Private/Scene/Scene.h
--- a/Engine/Plugins/Experimental/GPULightmass/Source/GPULightmass/Private/Scene/Scene.h
+++ b/Engine/Plugins/Experimental/GPULightmass/Source/GPULightmass/Private/Scene/Scene.h
@@ -119,7 +119,7 @@ public:
 
 	int32 GetPrimitiveIdForGPUScene(const FGeometryInstanceRenderStateRef& GeometryInstanceRef) const;
 
-	void SetupRayTracingScene(int32 LODIndex = 0);
+	bool SetupRayTracingScene(int32 LODIndex = 0);
 	void DestroyRayTracingScene();
 
 	void CalculateDistributionPrefixSumForAllLightmaps();
diff --git a/Engine/Plugins/Experimental/GPULightmass/Source/GPULightmass/Private/VolumetricLightmap.cpp b/Engine/Plugins/Experimental/GPULightmass/Source/GPULightmass/Private/VolumetricLightmap.cpp
--- a/Engine/Plugins/Experimental/GPULightmass/Source/GPULightmass/Private/VolumetricLightmap.cpp
+++ b/Engine/Plugins/Experimental/GPULightmass/Source/GPULightmass/Private/VolumetricLightmap.cpp
@@ -235,7 +235,11 @@ void FVolumetricLightmapRenderer::VoxelizeScene()
 	}
 
 	// Setup ray tracing scene with LOD 0
-	Scene->SetupRayTracingScene();
+	if (!Scene->SetupRayTracingScene())
+	{
+		GraphBuilder.Execute();
+		return;
+	}
 
 	auto* PassParameters = GraphBuilder.AllocParameters<FVoxelizeMeshPassParameters>();
 	PassParameters->PassUniformBuffer = GraphBuilder.CreateUniformBuffer(VLMVoxelizationParams);
@@ -499,7 +503,10 @@ void FVolumetricLightmapRenderer::BackgroundTick()
 
 	if (IsRayTracingEnabled())
 	{
-		Scene->SetupRayTracingScene();
+		if (!Scene->SetupRayTracingScene())
+		{
+			return;
+		}
 	}
 
 	FRDGBuilder GraphBuilder(FRHICommandListExecutor::GetImmediateCommandList());
