Store TilePosition+RelativeWorldPosition in FLightShaderConstants to support LWC.

diff --git a/Engine/Plugins/Experimental/GPULightmass/Shaders/Private/LightmapPathTracing.usf b/Engine/Plugins/Experimental/GPULightmass/Shaders/Private/LightmapPathTracing.usf
--- a/Engine/Plugins/Experimental/GPULightmass/Shaders/Private/LightmapPathTracing.usf
+++ b/Engine/Plugins/Experimental/GPULightmass/Shaders/Private/LightmapPathTracing.usf
@@ -927,7 +927,7 @@ uint MortonEncode3(uint3 Pixel)
 // Must match C++
 struct FLightShaderConstants
 {
-	float3 WorldPosition;
+	float3 RelativeWorldPosition;
 	float  InvRadius;
 	float3 Color;
 	float  FalloffExponent;
@@ -935,17 +935,19 @@ struct FLightShaderConstants
 	float  SpecularScale;
 	float3 Tangent;
 	float  SourceRadius;
+	float3 TilePosition;
+	float  SourceLength;
 	float2 SpotAngles;
 	float  SoftSourceRadius;
-	float  SourceLength;
 	float  RectLightBarnCosAngle;
 	float  RectLightBarnLength;
 	
 	void FillLightShaderParameters(inout FLightShaderParameters LightShaderParameters)
 	{
-		LightShaderParameters.Position = WorldPosition;
-		LightShaderParameters.TranslatedWorldPosition = LWCHackToFloat(PrimaryView.PreViewTranslation) + WorldPosition;
+		const FLWCVector3 WorldPosition = MakeLWCVector3(TilePosition, RelativeWorldPosition);
 
+		LightShaderParameters.Position = LWCHackToFloat(WorldPosition); // RT_LWC_TODO
+		LightShaderParameters.TranslatedWorldPosition = LWCToFloat(LWCAdd(WorldPosition, PrimaryView.PreViewTranslation));
 		LightShaderParameters.InvRadius = InvRadius;
 		LightShaderParameters.Color = Color;
 		LightShaderParameters.FalloffExponent = FalloffExponent;
diff --git a/Engine/Plugins/Experimental/GPULightmass/Source/GPULightmass/Private/Scene/Lights.h b/Engine/Plugins/Experimental/GPULightmass/Source/GPULightmass/Private/Scene/Lights.h
--- a/Engine/Plugins/Experimental/GPULightmass/Source/GPULightmass/Private/Scene/Lights.h
+++ b/Engine/Plugins/Experimental/GPULightmass/Source/GPULightmass/Private/Scene/Lights.h
@@ -15,7 +15,7 @@ class FSkyLightImportanceSamplingData;
 
 struct FLightShaderConstants
 {
-	FVector3f	WorldPosition;
+	FVector3f	RelativeWorldPosition;
 	float		InvRadius;
 	FVector3f	Color;
 	float		FalloffExponent;
@@ -23,9 +23,10 @@ struct FLightShaderConstants
 	float		SpecularScale;
 	FVector3f	Tangent;
 	float		SourceRadius;
+	FVector3f	TilePosition;
+	float		SourceLength;
 	FVector2f	SpotAngles;
 	float		SoftSourceRadius;
-	float		SourceLength;
 	float		RectLightBarnCosAngle;
 	float		RectLightBarnLength;
 
@@ -33,7 +34,9 @@ struct FLightShaderConstants
 
 	FLightShaderConstants(const FLightRenderParameters& LightShaderParameters)
 	{
-		WorldPosition = FVector3f(LightShaderParameters.WorldPosition); // LWC_TODO
+		const FLargeWorldRenderPosition AbsoluteWorldPosition(LightShaderParameters.WorldPosition);
+
+		RelativeWorldPosition = AbsoluteWorldPosition.GetOffset();
 		InvRadius = LightShaderParameters.InvRadius;
 		Color = FVector3f(LightShaderParameters.Color);
 		FalloffExponent = LightShaderParameters.FalloffExponent;
@@ -41,9 +44,10 @@ struct FLightShaderConstants
 		SpecularScale = LightShaderParameters.SpecularScale;
 		Tangent = LightShaderParameters.Tangent;
 		SourceRadius = LightShaderParameters.SourceRadius;
+		TilePosition = AbsoluteWorldPosition.GetTile();
+		SourceLength = LightShaderParameters.SourceLength;
 		SpotAngles = LightShaderParameters.SpotAngles;
 		SoftSourceRadius = LightShaderParameters.SoftSourceRadius;
-		SourceLength = LightShaderParameters.SourceLength;
 		RectLightBarnCosAngle = LightShaderParameters.RectLightBarnCosAngle;
 		RectLightBarnLength = LightShaderParameters.RectLightBarnLength;
 	}
