Clean rect. light atlas invalid texture slot value.

diff --git a/Engine/Plugins/Experimental/GPULightmass/Source/GPULightmass/Private/Scene/Lights.cpp b/Engine/Plugins/Experimental/GPULightmass/Source/GPULightmass/Private/Scene/Lights.cpp
--- a/Engine/Plugins/Experimental/GPULightmass/Source/GPULightmass/Private/Scene/Lights.cpp
+++ b/Engine/Plugins/Experimental/GPULightmass/Source/GPULightmass/Private/Scene/Lights.cpp
@@ -234,7 +234,9 @@ FLightRenderParameters FDirectionalLightRenderState::GetLightShaderParameters()
 	LightParameters.SourceRadius = FMath::Sin(0.5f * FMath::DegreesToRadians(LightSourceAngle));
 	LightParameters.SoftSourceRadius = 0; // Irrelevant when tracing shadow rays. FMath::Sin(0.5f * FMath::DegreesToRadians(LightSourceSoftAngle));
 	LightParameters.SourceLength = 0.0f;
-	LightParameters.SourceTexture = GWhiteTexture->TextureRHI; // Irrelevant when tracing shadow rays
+	LightParameters.RectLightAtlasUVOffset = FVector2f::ZeroVector;
+	LightParameters.RectLightAtlasUVScale = FVector2f::ZeroVector;
+	LightParameters.RectLightAtlasMaxLevel = FLightRenderParameters::GetRectLightAtlasInvalidMIPLevel();
 
 	return LightParameters;
 }
@@ -249,7 +251,9 @@ FLightRenderParameters FPointLightRenderState::GetLightShaderParameters() const
 	LightParameters.InvRadius = 1.0f / AttenuationRadius;
 	LightParameters.Color = Color;
 	LightParameters.SourceRadius = SourceRadius;
-
+	LightParameters.RectLightAtlasUVOffset = FVector2f::ZeroVector;
+	LightParameters.RectLightAtlasUVScale = FVector2f::ZeroVector;
+	LightParameters.RectLightAtlasMaxLevel = FLightRenderParameters::GetRectLightAtlasInvalidMIPLevel();
 	return LightParameters;
 }
 
@@ -264,6 +268,9 @@ FLightRenderParameters FSpotLightRenderState::GetLightShaderParameters() const
 	LightParameters.InvRadius = 1.0f / AttenuationRadius;
 	LightParameters.Color = Color;
 	LightParameters.SourceRadius = SourceRadius;
+	LightParameters.RectLightAtlasUVOffset = FVector2f::ZeroVector;
+	LightParameters.RectLightAtlasUVScale = FVector2f::ZeroVector;
+	LightParameters.RectLightAtlasMaxLevel = FLightRenderParameters::GetRectLightAtlasInvalidMIPLevel();
 
 	return LightParameters;
 }
@@ -283,9 +290,11 @@ FLightRenderParameters FRectLightRenderState::GetLightShaderParameters() const
 
 	LightParameters.SourceRadius = SourceWidth * 0.5f;
 	LightParameters.SourceLength = SourceHeight * 0.5f;
-	LightParameters.SourceTexture =  GWhiteTexture->TextureRHI;
 	LightParameters.RectLightBarnCosAngle = FMath::Cos(FMath::DegreesToRadians(BarnDoorAngle));
 	LightParameters.RectLightBarnLength = BarnDoorLength;
+	LightParameters.RectLightAtlasUVOffset = FVector2f::ZeroVector;
+	LightParameters.RectLightAtlasUVScale = FVector2f::ZeroVector;
+	LightParameters.RectLightAtlasMaxLevel = FLightRenderParameters::GetRectLightAtlasInvalidMIPLevel();
 
 	return LightParameters;
 }
diff --git a/Engine/Shaders/Private/RectLight.ush b/Engine/Shaders/Private/RectLight.ush
--- a/Engine/Shaders/Private/RectLight.ush
+++ b/Engine/Shaders/Private/RectLight.ush
@@ -10,6 +10,8 @@
 #define USE_SOURCE_TEXTURE 1
 #endif
 
+#define MAX_RECT_ATLAS_MIP 32
+
 struct FRect
 {
 	float3		Origin;
@@ -31,7 +33,7 @@ FRectTexture InitRectTexture()
 	FRectTexture Out;
 	Out.AtlasUVOffset = 0;
 	Out.AtlasUVScale = 0;
-	Out.AtlasMaxLevel = 99;
+	Out.AtlasMaxLevel = MAX_RECT_ATLAS_MIP;
 	return Out;
 }
 
@@ -47,7 +49,7 @@ FRectTexture InitRectTexture(half2 InUVOffset, half2 InUVScale, half InMaxLevel)
 float3 SampleRectTexture(FRectTexture RectTexture, float2 RectUV, float Level, bool bIsReference = false)
 {
 #if USE_SOURCE_TEXTURE
-	const bool bIsValid = RectTexture.AtlasMaxLevel < 16;
+	const bool bIsValid = RectTexture.AtlasMaxLevel < MAX_RECT_ATLAS_MIP;
 	const float2 RectTextureSize = RectTexture.AtlasUVScale * View.RectLightAtlasSizeAndInvSize.xy;
 	Level += log2(min(RectTextureSize.x, RectTextureSize.y)) - 2.f;
 	Level  = min(Level, RectTexture.AtlasMaxLevel);
diff --git a/Engine/Shaders/Private/VolumetricFog.usf b/Engine/Shaders/Private/VolumetricFog.usf
--- a/Engine/Shaders/Private/VolumetricFog.usf
+++ b/Engine/Shaders/Private/VolumetricFog.usf
@@ -722,7 +722,7 @@ void LightScatteringCS(
 				LightData.RectLightBarnLength = LocalLight.RectBarnDoorAndVirtualShadowMapId.y;
 				LightData.RectLightAtlasUVOffset = 0;
 				LightData.RectLightAtlasUVScale  = 0;
-				LightData.RectLightAtlasMaxLevel = 99;
+				LightData.RectLightAtlasMaxLevel = MAX_RECT_ATLAS_MIP;
 
 				LightData.bRadialLight = true;
 				LightData.bSpotLight = LightData.SpotAngles.x > -2.0f;
diff --git a/Engine/Source/Runtime/Engine/Private/Components/DirectionalLightComponent.cpp b/Engine/Source/Runtime/Engine/Private/Components/DirectionalLightComponent.cpp
--- a/Engine/Source/Runtime/Engine/Private/Components/DirectionalLightComponent.cpp
+++ b/Engine/Source/Runtime/Engine/Private/Components/DirectionalLightComponent.cpp
@@ -390,7 +390,7 @@ public:
 		LightParameters.SourceLength = 0.0f;
 		LightParameters.RectLightAtlasUVOffset = FVector2f::ZeroVector;
 		LightParameters.RectLightAtlasUVScale = FVector2f::ZeroVector;
-		LightParameters.RectLightAtlasMaxLevel = 99.f;
+		LightParameters.RectLightAtlasMaxLevel = FLightRenderParameters::GetRectLightAtlasInvalidMIPLevel();
 	}
 
 	virtual float GetLightSourceAngle() const override
diff --git a/Engine/Source/Runtime/Engine/Private/Components/PointLightComponent.cpp b/Engine/Source/Runtime/Engine/Private/Components/PointLightComponent.cpp
--- a/Engine/Source/Runtime/Engine/Private/Components/PointLightComponent.cpp
+++ b/Engine/Source/Runtime/Engine/Private/Components/PointLightComponent.cpp
@@ -49,7 +49,7 @@ void FPointLightSceneProxy::GetLightShaderParameters(FLightRenderParameters& Lig
 	LightParameters.RectLightBarnLength = -2.0f;
 	LightParameters.RectLightAtlasUVOffset = FVector2f::ZeroVector;
 	LightParameters.RectLightAtlasUVScale = FVector2f::ZeroVector;
-	LightParameters.RectLightAtlasMaxLevel = 99.f;
+	LightParameters.RectLightAtlasMaxLevel = FLightRenderParameters::GetRectLightAtlasInvalidMIPLevel();
 }
 
 /**
diff --git a/Engine/Source/Runtime/Engine/Private/Components/RectLightComponent.cpp b/Engine/Source/Runtime/Engine/Private/Components/RectLightComponent.cpp
--- a/Engine/Source/Runtime/Engine/Private/Components/RectLightComponent.cpp
+++ b/Engine/Source/Runtime/Engine/Private/Components/RectLightComponent.cpp
@@ -240,7 +240,7 @@ void FRectLightSceneProxy::GetLightShaderParameters(FLightRenderParameters& Ligh
 	LightParameters.RectLightBarnLength = BarnDoorLength;
 	LightParameters.RectLightAtlasUVOffset = FVector2f::ZeroVector;
 	LightParameters.RectLightAtlasUVScale = FVector2f::ZeroVector;
-	LightParameters.RectLightAtlasMaxLevel = 99;
+	LightParameters.RectLightAtlasMaxLevel = FLightRenderParameters::GetRectLightAtlasInvalidMIPLevel();
 	if (AtlasSlotIndex != ~0u)
 	{
 		GetSceneInterface()->GetRectLightAtlasSlot(this, &LightParameters);
diff --git a/Engine/Source/Runtime/Engine/Private/Components/SpotLightComponent.cpp b/Engine/Source/Runtime/Engine/Private/Components/SpotLightComponent.cpp
--- a/Engine/Source/Runtime/Engine/Private/Components/SpotLightComponent.cpp
+++ b/Engine/Source/Runtime/Engine/Private/Components/SpotLightComponent.cpp
@@ -69,7 +69,7 @@ public:
 		LightParameters.RectLightBarnLength = -2.0f;
 		LightParameters.RectLightAtlasUVOffset = FVector2f::ZeroVector;
 		LightParameters.RectLightAtlasUVScale = FVector2f::ZeroVector;
-		LightParameters.RectLightAtlasMaxLevel = 99.f;
+		LightParameters.RectLightAtlasMaxLevel = FLightRenderParameters::GetRectLightAtlasInvalidMIPLevel();
 	}
 
 	// FLightSceneInfo interface.
diff --git a/Engine/Source/Runtime/Engine/Public/SceneManagement.h b/Engine/Source/Runtime/Engine/Public/SceneManagement.h
--- a/Engine/Source/Runtime/Engine/Public/SceneManagement.h
+++ b/Engine/Source/Runtime/Engine/Public/SceneManagement.h
@@ -1354,9 +1354,6 @@ struct FLightRenderParameters
 {
 	ENGINE_API void MakeShaderParameters(const FViewMatrices& ViewMatrices, FLightShaderParameters& OutShaderParameters) const;
 
-	// Texture of the rect light.
-	FRHITexture* SourceTexture = nullptr;
-
 	// Position of the light in world space.
 	FVector WorldPosition;
 
@@ -1400,7 +1397,10 @@ struct FLightRenderParameters
 	// Rect. light atlas transformation
 	FVector2f RectLightAtlasUVOffset;
 	FVector2f RectLightAtlasUVScale;
-	float RectLightAtlasMaxLevel = 99.f;
+	float RectLightAtlasMaxLevel;
+
+	// Return Invalid rect light atlas MIP level
+	static float GetRectLightAtlasInvalidMIPLevel() { return 32.f;  }
 };
 
 /** 
diff --git a/Engine/Source/Runtime/Renderer/Private/LightRendering.cpp b/Engine/Source/Runtime/Renderer/Private/LightRendering.cpp
--- a/Engine/Source/Runtime/Renderer/Private/LightRendering.cpp
+++ b/Engine/Source/Runtime/Renderer/Private/LightRendering.cpp
@@ -296,7 +296,7 @@ FDeferredLightUniformStruct GetSimpleDeferredLightParameters(
 	Out.LightParameters.RectLightBarnLength = -2.0f;
 	Out.LightParameters.RectLightAtlasUVOffset = FVector2f::ZeroVector;
 	Out.LightParameters.RectLightAtlasUVScale = FVector2f::ZeroVector;
-	Out.LightParameters.RectLightAtlasMaxLevel = 99.f;
+	Out.LightParameters.RectLightAtlasMaxLevel = FLightRenderParameters::GetRectLightAtlasInvalidMIPLevel();
 	return Out;
 }
 FDeferredLightUniformStruct GetSimpleDeferredLightParameters(
diff --git a/Engine/Source/Runtime/Renderer/Private/RectLightTextureManager.cpp b/Engine/Source/Runtime/Renderer/Private/RectLightTextureManager.cpp
--- a/Engine/Source/Runtime/Renderer/Private/RectLightTextureManager.cpp
+++ b/Engine/Source/Runtime/Renderer/Private/RectLightTextureManager.cpp
@@ -1430,7 +1430,7 @@ FAtlasSlotDesc GetRectLightAtlasSlot(uint32 InSlotIndex)
 	FAtlasSlotDesc Out;
 	Out.UVOffset = FVector2f(0,0);
 	Out.UVScale = FVector2f(0, 0);
-	Out.MaxMipLevel = 99.f;
+	Out.MaxMipLevel = FLightRenderParameters::GetRectLightAtlasInvalidMIPLevel();
 
 	if (InSlotIndex < uint32(GRectLightTextureManager.AtlasSlots.Num()) && GRectLightTextureManager.AtlasTexture)
 	{
