MGPU:  Fix for crash in GPU Lightmass ray tracing acceleration structure initialization.  Needed to cut-and-paste logic from RayTracingScene.cpp, which separately initializes the acceleration structure addresses per GPU.

#rnx

diff --git a/Engine/Plugins/Experimental/GPULightmass/Source/GPULightmass/Private/LightmapRenderer.cpp b/Engine/Plugins/Experimental/GPULightmass/Source/GPULightmass/Private/LightmapRenderer.cpp
--- a/Engine/Plugins/Experimental/GPULightmass/Source/GPULightmass/Private/LightmapRenderer.cpp
+++ b/Engine/Plugins/Experimental/GPULightmass/Source/GPULightmass/Private/LightmapRenderer.cpp
@@ -1037,8 +1037,12 @@ bool FSceneRenderState::SetupRayTracingScene(int32 LODIndex)
 			FRWBufferStructured InstanceBuffer;
 			InstanceBuffer.Initialize(TEXT("LightmassRayTracingInstanceBuffer"), GRHIRayTracingInstanceDescriptorSize, SceneInitializer.NumNativeInstances);
 
+			// Need to pass "BUF_MultiGPUAllocate", as acceleration structure virtual addresses are different per GPU
 			FByteAddressBuffer AccelerationStructureAddressesBuffer;
-			AccelerationStructureAddressesBuffer.Initialize(TEXT("LightmassRayTracingAccelerationStructureAddressesBuffer"), SceneInitializer.ReferencedGeometries.Num() * sizeof(FRayTracingAccelerationStructureAddress), BUF_Volatile);
+			AccelerationStructureAddressesBuffer.Initialize(
+				TEXT("LightmassRayTracingAccelerationStructureAddressesBuffer"),
+				SceneInitializer.ReferencedGeometries.Num() * sizeof(FRayTracingAccelerationStructureAddress),
+				BUF_Volatile | BUF_MultiGPUAllocate);
 
 			const uint32 InstanceUploadBufferSize = SceneInitializer.NumNativeInstances * sizeof(FRayTracingInstanceDescriptorInput);
 			FBufferRHIRef InstanceUploadBuffer;
@@ -1074,20 +1078,27 @@ bool FSceneRenderState::SetupRayTracingScene(int32 LODIndex)
 				RHICmdList.UnlockBuffer(InstanceUploadBuffer);
 			}
 
-			RHICmdList.EnqueueLambda([BufferRHIRef = AccelerationStructureAddressesBuffer.Buffer, &SceneInitializer](FRHICommandListImmediate& RHICmdList)
-				{
-					FRayTracingAccelerationStructureAddress* AddressesPtr = (FRayTracingAccelerationStructureAddress*)RHICmdList.LockBuffer(
-						BufferRHIRef,
-						0,
-						SceneInitializer.ReferencedGeometries.Num() * sizeof(FRayTracingAccelerationStructureAddress), RLM_WriteOnly);
+			// Pull this out here, because command list playback (where the lambda is executed) doesn't update the GPU mask
+			FRHIGPUMask IterateGPUMasks = RHICmdList.GetGPUMask();
 
-					const uint32 NumGeometries = SceneInitializer.ReferencedGeometries.Num();
-					for (uint32 GeometryIndex = 0; GeometryIndex < NumGeometries; ++GeometryIndex)
+			RHICmdList.EnqueueLambda([BufferRHIRef = AccelerationStructureAddressesBuffer.Buffer, &SceneInitializer, IterateGPUMasks](FRHICommandListImmediate& RHICmdList)
+				{
+					for (uint32 GPUIndex : IterateGPUMasks)
 					{
-						AddressesPtr[GeometryIndex] = SceneInitializer.ReferencedGeometries[GeometryIndex]->GetAccelerationStructureAddress(RHICmdList.GetGPUMask().ToIndex());
-					}
+						FRayTracingAccelerationStructureAddress* AddressesPtr = (FRayTracingAccelerationStructureAddress*)RHICmdList.LockBufferMGPU(
+							BufferRHIRef,
+							GPUIndex,
+							0,
+							SceneInitializer.ReferencedGeometries.Num() * sizeof(FRayTracingAccelerationStructureAddress), RLM_WriteOnly);
+
+						const uint32 NumGeometries = SceneInitializer.ReferencedGeometries.Num();
+						for (uint32 GeometryIndex = 0; GeometryIndex < NumGeometries; ++GeometryIndex)
+						{
+							AddressesPtr[GeometryIndex] = SceneInitializer.ReferencedGeometries[GeometryIndex]->GetAccelerationStructureAddress(GPUIndex);
+						}
 
-					RHICmdList.UnlockBuffer(BufferRHIRef);
+						RHICmdList.UnlockBufferMGPU(BufferRHIRef, GPUIndex);
+					}
 				});
 
 			BuildRayTracingInstanceBuffer(
