GPULightmass plugin tagged as beta

diff --git a/Engine/Plugins/Experimental/GPULightmass/GPULightmass.uplugin b/Engine/Plugins/Experimental/GPULightmass/GPULightmass.uplugin
--- a/Engine/Plugins/Experimental/GPULightmass/GPULightmass.uplugin
+++ b/Engine/Plugins/Experimental/GPULightmass/GPULightmass.uplugin
@@ -12,7 +12,7 @@
 	"SupportURL" : "",
 	"EnabledByDefault" : false,
 	"CanContainContent" : true,
-	"IsBetaVersion" : false,
+	"IsBetaVersion" : true,
 	"Installed" : false,
 	"Modules" :
 	[
