Fixed crash in GPU Lightmass when "Allow Static Lighting" is disabled in the project settings. Now GPU Lightmass only runs when that option is checked.

diff --git a/Engine/Plugins/Experimental/GPULightmass/Source/GPULightmassEditor/Private/GPULightmassEditorModule.cpp b/Engine/Plugins/Experimental/GPULightmass/Source/GPULightmassEditor/Private/GPULightmassEditorModule.cpp
--- a/Engine/Plugins/Experimental/GPULightmass/Source/GPULightmassEditor/Private/GPULightmassEditorModule.cpp
+++ b/Engine/Plugins/Experimental/GPULightmass/Source/GPULightmassEditor/Private/GPULightmassEditorModule.cpp
@@ -107,7 +107,7 @@ TSharedRef<SDockTab> FGPULightmassEditorModule::SpawnSettingsTab(const FSpawnTab
 					SNew(SButton)
 					.HAlign(HAlign_Center)
 					.ButtonStyle(FEditorStyle::Get(), "FlatButton.Success")
-					.IsEnabled(IsRayTracingEnabled())
+					.IsEnabled(IsRayTracingEnabled() && IsAllowStaticLightingEnabled())
 					.Visibility_Lambda([](){ return IsRunning() ? EVisibility::Collapsed : EVisibility::Visible; })
 					.OnClicked_Raw(this, &FGPULightmassEditorModule::OnStartClicked)
 					[
@@ -231,6 +231,10 @@ TSharedRef<SDockTab> FGPULightmassEditorModule::SpawnSettingsTab(const FSpawnTab
 					{
 						return LOCTEXT("GPULightmassRayTracingDisabled", "GPU Lightmass requires ray tracing support which is disabled.");
 					}
+					else if (!IsAllowStaticLightingEnabled())
+					{
+						return LOCTEXT("GPULightmassAllowStaticLightingDisabled", "GPU Lightmass requires Allow Static Lighting enabled in the project settings.");
+					}
 
 					// Ready
 					static FText ReadyMsg = FText(LOCTEXT("GPULightmassReady", "GPU Lightmass is ready."));
@@ -318,6 +322,12 @@ bool FGPULightmassEditorModule::IsRunning()
 	return false;
 }
 
+bool FGPULightmassEditorModule::IsAllowStaticLightingEnabled()
+{
+	static const auto AllowStaticLightingVar = IConsoleManager::Get().FindTConsoleVariableDataInt(TEXT("r.AllowStaticLighting"));
+	return (AllowStaticLightingVar && AllowStaticLightingVar->GetValueOnAnyThread() > 0);
+}
+
 FReply FGPULightmassEditorModule::OnStartClicked()
 {
 	if (UWorld* World = GEditor->GetEditorWorldContext().World())
diff --git a/Engine/Plugins/Experimental/GPULightmass/Source/GPULightmassEditor/Private/GPULightmassEditorModule.h b/Engine/Plugins/Experimental/GPULightmass/Source/GPULightmassEditor/Private/GPULightmassEditorModule.h
--- a/Engine/Plugins/Experimental/GPULightmass/Source/GPULightmassEditor/Private/GPULightmassEditorModule.h
+++ b/Engine/Plugins/Experimental/GPULightmass/Source/GPULightmassEditor/Private/GPULightmassEditorModule.h
@@ -36,6 +36,7 @@ public:
 	static bool IsRealtimeOn();
 	static bool IsRunning();
 	static bool IsBakeWhatYouSeeMode();
+	static bool IsAllowStaticLightingEnabled();
 
 	FReply OnStartClicked();
 	FReply OnSaveAndStopClicked();
