Fixed crash in GPU Lightmass when debug mode was enabled

diff --git a/Engine/Plugins/Experimental/GPULightmass/Source/GPULightmass/Private/LightmapRenderer.cpp b/Engine/Plugins/Experimental/GPULightmass/Source/GPULightmass/Private/LightmapRenderer.cpp
--- a/Engine/Plugins/Experimental/GPULightmass/Source/GPULightmass/Private/LightmapRenderer.cpp
+++ b/Engine/Plugins/Experimental/GPULightmass/Source/GPULightmass/Private/LightmapRenderer.cpp
@@ -3665,9 +3665,7 @@ void FLightmapRenderer::RenderIrradianceCacheVisualization(FPostOpaqueRenderPara
 	FRDGBuilder& GraphBuilder = *Parameters.GraphBuilder;
 
 	auto* PassParameters = GraphBuilder.AllocParameters<FVisualizeIrradianceCachePS::FParameters>();
-	TUniformBufferRef<FViewUniformShaderParameters> Ref;
-	*Ref.GetInitReference() = Parameters.ViewUniformBuffer;
-	PassParameters->View = Ref;
+	PassParameters->View = Parameters.View->ViewUniformBuffer;
 	PassParameters->SceneTextures = Parameters.SceneTexturesUniformParams;
 	PassParameters->IrradianceCachingParameters = Scene->IrradianceCache->IrradianceCachingParametersUniformBuffer;
 	PassParameters->RenderTargets[0] = FRenderTargetBinding(Parameters.ColorTexture, ERenderTargetLoadAction::ELoad);
diff --git a/Engine/Source/Runtime/RenderCore/Public/RendererInterface.h b/Engine/Source/Runtime/RenderCore/Public/RendererInterface.h
--- a/Engine/Source/Runtime/RenderCore/Public/RendererInterface.h
+++ b/Engine/Source/Runtime/RenderCore/Public/RendererInterface.h
@@ -34,6 +34,7 @@ class FShader;
 class FShaderMapPointerTable;
 class FRDGBuilder;
 class FMaterialRenderProxy;
+class FViewInfo;
 template<typename ShaderType, typename PointerTableType> class TShaderRefBase;
 
 namespace Nanite
@@ -625,6 +626,7 @@ public:
 	TRDGUniformBufferRef<FMobileSceneTextureUniformParameters> MobileSceneTexturesUniformParams = nullptr;
 	const FGlobalDistanceFieldParameterData* GlobalDistanceFieldParams = nullptr;
 	void* Uid = nullptr; // A unique identifier for the view.
+	const FViewInfo* View = nullptr; 
 };
 DECLARE_MULTICAST_DELEGATE_OneParam(FOnPostOpaqueRender, class FPostOpaqueRenderParameters&);
 typedef FOnPostOpaqueRender::FDelegate FPostOpaqueRenderDelegate;
diff --git a/Engine/Source/Runtime/Renderer/Private/SceneRendering.cpp b/Engine/Source/Runtime/Renderer/Private/SceneRendering.cpp
--- a/Engine/Source/Runtime/Renderer/Private/SceneRendering.cpp
+++ b/Engine/Source/Runtime/Renderer/Private/SceneRendering.cpp
@@ -4464,6 +4464,7 @@ void FRendererModule::RenderPostOpaqueExtensions(
 			RenderParameters.GraphBuilder = &GraphBuilder;
 
 			RenderParameters.Uid = (void*)(&View);
+			RenderParameters.View = &View;
 			PostOpaqueRenderDelegate.Broadcast(RenderParameters);
 		}
 	}
@@ -4495,6 +4496,7 @@ void FRendererModule::RenderOverlayExtensions(
 			RenderParameters.GraphBuilder = &GraphBuilder;
 
 			RenderParameters.Uid = (void*)(&View);
+			RenderParameters.View = &View;
 			OverlayRenderDelegate.Broadcast(RenderParameters);
 		}
 	}
