Fixing GPULightmass to use the FeatureLevel from the World.

diff --git a/Engine/Plugins/Experimental/GPULightmass/Source/GPULightmass/Private/LightmapRenderer.cpp b/Engine/Plugins/Experimental/GPULightmass/Source/GPULightmass/Private/LightmapRenderer.cpp
--- a/Engine/Plugins/Experimental/GPULightmass/Source/GPULightmass/Private/LightmapRenderer.cpp
+++ b/Engine/Plugins/Experimental/GPULightmass/Source/GPULightmass/Private/LightmapRenderer.cpp
@@ -614,7 +614,7 @@ void FCachedRayTracingSceneData::SetupFromSceneRenderState(FSceneRenderState& Sc
 			for (int32 SegmentIndex = 0; SegmentIndex < MeshBatches.Num(); SegmentIndex++)
 			{
 				const FMaterialRenderProxy* FallbackMaterialRenderProxyPtr = nullptr;
-				const FMaterial& Material = MeshBatches[SegmentIndex].MaterialRenderProxy->GetMaterialWithFallback(GMaxRHIFeatureLevel, FallbackMaterialRenderProxyPtr);
+				const FMaterial& Material = MeshBatches[SegmentIndex].MaterialRenderProxy->GetMaterialWithFallback(Scene.FeatureLevel, FallbackMaterialRenderProxyPtr);
 
 				bAllSegmentsUnlit &= Material.GetShadingModels().HasOnlyShadingModel(MSM_Unlit) || !MeshBatches[SegmentIndex].CastShadow;
 				bAllSegmentsOpaque &= Material.GetBlendMode() == EBlendMode::BLEND_Opaque;
@@ -667,7 +667,7 @@ void FCachedRayTracingSceneData::SetupFromSceneRenderState(FSceneRenderState& Sc
 			for (int32 SegmentIndex = 0; SegmentIndex < MeshBatches.Num(); SegmentIndex++)
 			{
 				const FMaterialRenderProxy* FallbackMaterialRenderProxyPtr = nullptr;
-				const FMaterial& Material = MeshBatches[SegmentIndex].MaterialRenderProxy->GetMaterialWithFallback(GMaxRHIFeatureLevel, FallbackMaterialRenderProxyPtr);
+				const FMaterial& Material = MeshBatches[SegmentIndex].MaterialRenderProxy->GetMaterialWithFallback(Scene.FeatureLevel, FallbackMaterialRenderProxyPtr);
 
 				bAllSegmentsUnlit &= Material.GetShadingModels().HasOnlyShadingModel(MSM_Unlit) || !MeshBatches[SegmentIndex].CastShadow;
 				bAllSegmentsOpaque &= Material.GetBlendMode() == EBlendMode::BLEND_Opaque;
@@ -982,7 +982,7 @@ bool FSceneRenderState::SetupRayTracingScene(int32 LODIndex)
 							RayTracingMeshProcessor.AddMeshBatch(MeshBatches[SegmentIndex], 1, nullptr);
 
 							const FMaterialRenderProxy* FallbackMaterialRenderProxyPtr = nullptr;
-							const FMaterial& Material = MeshBatches[SegmentIndex].MaterialRenderProxy->GetMaterialWithFallback(GMaxRHIFeatureLevel, FallbackMaterialRenderProxyPtr);
+							const FMaterial& Material = MeshBatches[SegmentIndex].MaterialRenderProxy->GetMaterialWithFallback(FeatureLevel, FallbackMaterialRenderProxyPtr);
 
 							bAllSegmentsUnlit &= Material.GetShadingModels().HasOnlyShadingModel(MSM_Unlit) || !MeshBatches[SegmentIndex].CastShadow;
 							bAllSegmentsOpaque &= Material.GetBlendMode() == EBlendMode::BLEND_Opaque;
@@ -1135,24 +1135,26 @@ bool FSceneRenderState::SetupRayTracingScene(int32 LODIndex)
 			PSOInitializer.MaxPayloadSizeInBytes = RAY_TRACING_MAX_ALLOWED_PAYLOAD_SIZE;
 			PSOInitializer.bAllowHitGroupIndexing = true;
 
+			FGlobalShaderMap* GlobalShaderMap = GetGlobalShaderMap(FeatureLevel);
+
 			TArray<FRHIRayTracingShader*> RayGenShaderTable;
 			{
 				FLightmapPathTracingRGS::FPermutationDomain PermutationVector;
 				PermutationVector.Set<FLightmapPathTracingRGS::FUseFirstBounceRayGuiding>(LightmapRenderer->bUseFirstBounceRayGuiding);
 				PermutationVector.Set<FLightmapPathTracingRGS::FUseIrradianceCaching>(Settings->bUseIrradianceCaching);
-				RayGenShaderTable.Add(GetGlobalShaderMap(GMaxRHIFeatureLevel)->GetShader<FLightmapPathTracingRGS>(PermutationVector).GetRayTracingShader());
+				RayGenShaderTable.Add(GlobalShaderMap->GetShader<FLightmapPathTracingRGS>(PermutationVector).GetRayTracingShader());
 			}
 			{
-				RayGenShaderTable.Add(GetGlobalShaderMap(GMaxRHIFeatureLevel)->GetShader<FStationaryLightShadowTracingRGS>().GetRayTracingShader());
+				RayGenShaderTable.Add(GlobalShaderMap->GetShader<FStationaryLightShadowTracingRGS>().GetRayTracingShader());
 			}
 			{
 				FVolumetricLightmapPathTracingRGS::FPermutationDomain PermutationVector;
 				PermutationVector.Set<FVolumetricLightmapPathTracingRGS::FUseIrradianceCaching>(Settings->bUseIrradianceCaching);
-				RayGenShaderTable.Add(GetGlobalShaderMap(GMaxRHIFeatureLevel)->GetShader<FVolumetricLightmapPathTracingRGS>(PermutationVector).GetRayTracingShader());
+				RayGenShaderTable.Add(GlobalShaderMap->GetShader<FVolumetricLightmapPathTracingRGS>(PermutationVector).GetRayTracingShader());
 			}
 			PSOInitializer.SetRayGenShaderTable(RayGenShaderTable);
 
-			auto DefaultClosestHitShader = GetGlobalShaderMap(ERHIFeatureLevel::SM5)->GetShader<FOpaqueShadowHitGroup>().GetRayTracingShader();
+			auto DefaultClosestHitShader = GlobalShaderMap->GetShader<FOpaqueShadowHitGroup>().GetRayTracingShader();
 			TArray<FRHIRayTracingShader*> RayTracingMaterialLibrary;
 			FShaderMapResource::GetRayTracingMaterialLibrary(RayTracingMaterialLibrary, DefaultClosestHitShader);
 
@@ -1417,6 +1419,8 @@ void FLightmapRenderer::Finalize(FRDGBuilder& GraphBuilder)
 		return *GraphBuilder.AllocObject<FShaderResourceViewRHIRef>(View);
 	};
 
+	FGlobalShaderMap* GlobalShaderMap = GetGlobalShaderMap(Scene->FeatureLevel);
+
 	// Upload & copy converged tiles directly
 	{
 		TArray<FLightmapTileRequest> TileUploadRequests = PendingTileRequests.FilterByPredicate(
@@ -1644,7 +1648,7 @@ void FLightmapRenderer::Finalize(FRDGBuilder& GraphBuilder)
 					PassParameters->SrcTilePositions = HoldReference(GraphBuilder, SrcTilePositionsSRV);
 					PassParameters->DstTilePositions = HoldReference(GraphBuilder, DstTilePositionsSRV);
 
-					TShaderMapRef<FUploadConvergedLightmapTilesCS> ComputeShader(GetGlobalShaderMap(GMaxRHIFeatureLevel));
+					TShaderMapRef<FUploadConvergedLightmapTilesCS> ComputeShader(GlobalShaderMap);
 					FComputeShaderUtils::AddPass(
 						GraphBuilder,
 						RDG_EVENT_NAME("UploadConvergedLightmapTiles"),
@@ -1692,7 +1696,7 @@ void FLightmapRenderer::Finalize(FRDGBuilder& GraphBuilder)
 					PassParameters->SrcTilePositions = HoldReference(GraphBuilder, SrcTilePositionsSRV);
 					PassParameters->DstTilePositions = HoldReference(GraphBuilder, DstTilePositionsSRV);
 
-					TShaderMapRef<FUploadConvergedLightmapTilesCS> ComputeShader(GetGlobalShaderMap(GMaxRHIFeatureLevel));
+					TShaderMapRef<FUploadConvergedLightmapTilesCS> ComputeShader(GlobalShaderMap);
 					FComputeShaderUtils::AddPass(
 						GraphBuilder,
 						RDG_EVENT_NAME("UploadConvergedLightmapTiles"),
@@ -1740,7 +1744,7 @@ void FLightmapRenderer::Finalize(FRDGBuilder& GraphBuilder)
 					PassParameters->SrcTilePositions = HoldReference(GraphBuilder, SrcTilePositionsSRV);
 					PassParameters->DstTilePositions = HoldReference(GraphBuilder, DstTilePositionsSRV);
 
-					TShaderMapRef<FUploadConvergedLightmapTilesCS> ComputeShader(GetGlobalShaderMap(GMaxRHIFeatureLevel));
+					TShaderMapRef<FUploadConvergedLightmapTilesCS> ComputeShader(GlobalShaderMap);
 					FComputeShaderUtils::AddPass(
 						GraphBuilder,
 						RDG_EVENT_NAME("UploadConvergedLightmapTiles"),
@@ -2176,7 +2180,7 @@ void FLightmapRenderer::Finalize(FRDGBuilder& GraphBuilder)
 						Parameters->TilePositions = HoldReference(GraphBuilder, TilePositionsBuffer.SRV);
 						Parameters->TilePool = ScratchTilePoolLayerUAVs[ScratchLayerIndex];
 
-						TShaderMapRef<FMultiTileClearCS> ComputeShader(GetGlobalShaderMap(GMaxRHIFeatureLevel));
+						TShaderMapRef<FMultiTileClearCS> ComputeShader(GlobalShaderMap);
 						FComputeShaderUtils::AddPass(
 							GraphBuilder,
 							RDG_EVENT_NAME("MultiTileClear"),
@@ -2347,7 +2351,7 @@ void FLightmapRenderer::Finalize(FRDGBuilder& GraphBuilder)
 									FLightmapPathTracingRGS::FPermutationDomain PermutationVector;
 									PermutationVector.Set<FLightmapPathTracingRGS::FUseFirstBounceRayGuiding>(bUseFirstBounceRayGuiding);
 									PermutationVector.Set<FLightmapPathTracingRGS::FUseIrradianceCaching>(Scene->Settings->bUseIrradianceCaching);
-									auto RayGenerationShader = GetGlobalShaderMap(GMaxRHIFeatureLevel)->GetShader<FLightmapPathTracingRGS>(PermutationVector);
+									auto RayGenerationShader = GlobalShaderMap->GetShader<FLightmapPathTracingRGS>(PermutationVector);
 									ClearUnusedGraphResources(RayGenerationShader, PassParameters);
 
 									GraphBuilder.AddPass(
@@ -2375,7 +2379,7 @@ void FLightmapRenderer::Finalize(FRDGBuilder& GraphBuilder)
 									PassParameters->RayGuidingCDFY = GraphBuilder.CreateUAV(RayGuidingCDFY);
 									PassParameters->NumRayGuidingTrialSamples = NumFirstBounceRayGuidingTrialSamples;
 
-									TShaderMapRef<FFirstBounceRayGuidingCDFBuildCS> ComputeShader(GetGlobalShaderMap(GMaxRHIFeatureLevel));
+									TShaderMapRef<FFirstBounceRayGuidingCDFBuildCS> ComputeShader(GlobalShaderMap);
 									FComputeShaderUtils::AddPass(
 										GraphBuilder,
 										RDG_EVENT_NAME("FirstBounceRayGuidingCDFBuild"),
@@ -2436,7 +2440,7 @@ void FLightmapRenderer::Finalize(FRDGBuilder& GraphBuilder)
 						Parameters->TilePositions = HoldReference(GraphBuilder, TilePositionsBufferSRV);
 						Parameters->TilePool = ScratchTilePoolLayerUAVs[ScratchLayerIndex];
 
-						TShaderMapRef<FMultiTileClearCS> ComputeShader(GetGlobalShaderMap(GMaxRHIFeatureLevel));
+						TShaderMapRef<FMultiTileClearCS> ComputeShader(GlobalShaderMap);
 						FComputeShaderUtils::AddPass(
 							GraphBuilder,
 							RDG_EVENT_NAME("MultiTileClear"),
@@ -2714,7 +2718,7 @@ void FLightmapRenderer::Finalize(FRDGBuilder& GraphBuilder)
 						PassParameters->ShadowMask = GraphBuilder.CreateUAV(ShadowMask);
 						PassParameters->ShadowMaskSampleCount = GraphBuilder.CreateUAV(ShadowMaskSampleCount);
 
-						auto RayGenerationShader = GetGlobalShaderMap(GMaxRHIFeatureLevel)->GetShader<FStationaryLightShadowTracingRGS>();
+						auto RayGenerationShader = GlobalShaderMap->GetShader<FStationaryLightShadowTracingRGS>();
 						ClearUnusedGraphResources(RayGenerationShader, PassParameters);
 
 						GraphBuilder.AddPass(
@@ -2820,7 +2824,7 @@ void FLightmapRenderer::Finalize(FRDGBuilder& GraphBuilder)
 				PermutationVector.Set<FSelectiveLightmapOutputCS::FOutputLayerDim>(0);
 				PermutationVector.Set<FSelectiveLightmapOutputCS::FDrawProgressBars>(Scene->Settings->bShowProgressBars);
 
-				auto Shader = GetGlobalShaderMap(GMaxRHIFeatureLevel)->GetShader<FSelectiveLightmapOutputCS>(PermutationVector);
+				auto Shader = GlobalShaderMap->GetShader<FSelectiveLightmapOutputCS>(PermutationVector);
 
 				FSelectiveLightmapOutputCS::FParameters* PassParameters = GraphBuilder.AllocParameters<FSelectiveLightmapOutputCS::FParameters>();
 				PassParameters->NumBatchedTiles = GPUBatchedTileRequests.BatchedTilesDesc.Num();
@@ -2848,7 +2852,7 @@ void FLightmapRenderer::Finalize(FRDGBuilder& GraphBuilder)
 				PermutationVector.Set<FSelectiveLightmapOutputCS::FOutputLayerDim>(2);
 				PermutationVector.Set<FSelectiveLightmapOutputCS::FDrawProgressBars>(Scene->Settings->bShowProgressBars);
 
-				auto Shader = GetGlobalShaderMap(GMaxRHIFeatureLevel)->GetShader<FSelectiveLightmapOutputCS>(PermutationVector);
+				auto Shader = GlobalShaderMap->GetShader<FSelectiveLightmapOutputCS>(PermutationVector);
 
 				FSelectiveLightmapOutputCS::FParameters* PassParameters = GraphBuilder.AllocParameters<FSelectiveLightmapOutputCS::FParameters>();
 				PassParameters->NumBatchedTiles = GPUBatchedTileRequests.BatchedTilesDesc.Num();
@@ -2998,7 +3002,7 @@ void FLightmapRenderer::Finalize(FRDGBuilder& GraphBuilder)
 					PassParameters->StagingHQLayer1 = GraphBuilder.CreateUAV(StagingHQLayer1);
 					PassParameters->StagingShadowMask = GraphBuilder.CreateUAV(StagingShadowMask);
 
-					TShaderMapRef<FCopyConvergedLightmapTilesCS> ComputeShader(GetGlobalShaderMap(GMaxRHIFeatureLevel));
+					TShaderMapRef<FCopyConvergedLightmapTilesCS> ComputeShader(GlobalShaderMap);
 					FComputeShaderUtils::AddPass(
 						GraphBuilder,
 						RDG_EVENT_NAME("CopyConvergedLightmapTiles"),
@@ -3439,7 +3443,7 @@ void FLightmapRenderer::BackgroundTick()
 
 								Lightmap.LightmapPreviewVirtualTexture->ProducePageData(
 									RHICmdList,
-									ERHIFeatureLevel::SM5,
+									Scene->FeatureLevel,
 									EVTProducePageFlags::None,
 									FVirtualTextureProducerHandle(),
 									0b111,
@@ -3663,6 +3667,7 @@ void FLightmapRenderer::DeduplicateRecordedTileRequests()
 void FLightmapRenderer::RenderIrradianceCacheVisualization(FPostOpaqueRenderParameters& Parameters)
 {
 	FRDGBuilder& GraphBuilder = *Parameters.GraphBuilder;
+	const ERHIFeatureLevel::Type FeatureLevel = Scene->FeatureLevel;
 
 	auto* PassParameters = GraphBuilder.AllocParameters<FVisualizeIrradianceCachePS::FParameters>();
 	PassParameters->View = Parameters.View->ViewUniformBuffer;
@@ -3677,12 +3682,12 @@ void FLightmapRenderer::RenderIrradianceCacheVisualization(FPostOpaqueRenderPara
 		RDG_EVENT_NAME("ClearIrradiance"),
 		PassParameters,
 		ERDGPassFlags::Raster,
-		[ViewportRect, PassParameters, TextureExtent] (FRHICommandList& RHICmdList)
+		[ViewportRect, PassParameters, TextureExtent, FeatureLevel] (FRHICommandList& RHICmdList)
 	{
 		RHICmdList.SetViewport(0, 0, 0.0f, ViewportRect.Width(), ViewportRect.Height(), 1.0f);
 
-		TShaderMapRef<FPostProcessVS> VertexShader(GetGlobalShaderMap(GMaxRHIFeatureLevel));
-		TShaderMapRef<FVisualizeIrradianceCachePS> PixelShader(GetGlobalShaderMap(GMaxRHIFeatureLevel));
+		TShaderMapRef<FPostProcessVS> VertexShader(GetGlobalShaderMap(FeatureLevel));
+		TShaderMapRef<FVisualizeIrradianceCachePS> PixelShader(GetGlobalShaderMap(FeatureLevel));
 
 		FGraphicsPipelineStateInitializer GraphicsPSOInit;
 		RHICmdList.ApplyCachedRenderTargets(GraphicsPSOInit);
diff --git a/Engine/Plugins/Experimental/GPULightmass/Source/GPULightmass/Private/Scene/Scene.cpp b/Engine/Plugins/Experimental/GPULightmass/Source/GPULightmass/Private/Scene/Scene.cpp
--- a/Engine/Plugins/Experimental/GPULightmass/Source/GPULightmass/Private/Scene/Scene.cpp
+++ b/Engine/Plugins/Experimental/GPULightmass/Source/GPULightmass/Private/Scene/Scene.cpp
@@ -40,12 +40,14 @@ FScene::FScene(FGPULightmass* InGPULightmass)
 	: GPULightmass(InGPULightmass)
 	, Settings(InGPULightmass->Settings)
 	, Geometries(*this)
+	, FeatureLevel(InGPULightmass->World->FeatureLevel)
 {
 	StaticMeshInstances.LinkRenderStateArray(RenderState.StaticMeshInstanceRenderStates);
 	InstanceGroups.LinkRenderStateArray(RenderState.InstanceGroupRenderStates);
 	Landscapes.LinkRenderStateArray(RenderState.LandscapeRenderStates);
 
 	RenderState.Settings = Settings;
+	RenderState.FeatureLevel = FeatureLevel;
 
 	ENQUEUE_RENDER_COMMAND(RenderThreadInit)(
 		[&RenderState = RenderState](FRHICommandListImmediate&) mutable
@@ -872,7 +874,7 @@ void FScene::AddGeometryInstanceFromComponent(UStaticMeshComponent* InComponent)
 
 	for (auto ResourceCluster : ResourceClusters)
 	{
-		ResourceCluster->UpdateUniformBuffer(ERHIFeatureLevel::SM5);
+		ResourceCluster->UpdateUniformBuffer(FeatureLevel);
 	}
 
 	if (InComponent->GetWorld())
@@ -1023,7 +1025,7 @@ void FScene::AddGeometryInstanceFromComponent(UInstancedStaticMeshComponent* InC
 	FInstanceGroupRenderState InstanceRenderState;
 	InstanceRenderState.ComponentUObject = Instance->ComponentUObject;
 	InstanceRenderState.RenderData = Instance->ComponentUObject->GetStaticMesh()->GetRenderData();
-	InstanceRenderState.InstancedRenderData = MakeUnique<FInstancedStaticMeshRenderData>(Instance->ComponentUObject, ERHIFeatureLevel::SM5);
+	InstanceRenderState.InstancedRenderData = MakeUnique<FInstancedStaticMeshRenderData>(Instance->ComponentUObject, FeatureLevel);
 	InstanceRenderState.LocalToWorld = InComponent->GetRenderMatrix();
 	InstanceRenderState.WorldBounds = InComponent->Bounds;
 	InstanceRenderState.ActorPosition = InComponent->GetAttachmentRootActor() ? InComponent->GetAttachmentRootActor()->GetActorLocation() : FVector(ForceInitToZero);
@@ -1112,7 +1114,7 @@ void FScene::AddGeometryInstanceFromComponent(UInstancedStaticMeshComponent* InC
 
 	for (auto ResourceCluster : ResourceClusters)
 	{
-		ResourceCluster->UpdateUniformBuffer(ERHIFeatureLevel::SM5);
+		ResourceCluster->UpdateUniformBuffer(FeatureLevel);
 	}
 }
 
@@ -1241,9 +1243,8 @@ void FScene::AddGeometryInstanceFromComponent(ULandscapeComponent* InComponent)
 
 	const int8 SubsectionSizeLog2 = FMath::CeilLogTwo(InComponent->SubsectionSizeQuads + 1);
 	InstanceRenderState.SharedBuffersKey = (SubsectionSizeLog2 & 0xf) | ((InComponent->NumSubsections & 0xf) << 4) |
-		(InComponent->GetWorld()->FeatureLevel <= ERHIFeatureLevel::ES3_1 ? 0 : (1 << 30)) | (InComponent->XYOffsetmapTexture == nullptr ? 0 : 1 << 31);
+		(FeatureLevel <= ERHIFeatureLevel::ES3_1 ? 0 : (1 << 30)) | (InComponent->XYOffsetmapTexture == nullptr ? 0 : 1 << 31);
 	InstanceRenderState.SharedBuffersKey |= 1 << 29; // Use this bit to indicate it is GPULightmass specific buffer (which only has FixedGridVertexFactory created)
-	TEnumAsByte<ERHIFeatureLevel::Type> FeatureLevel = InComponent->GetWorld()->FeatureLevel;
 
 	TArray<UMaterialInterface*> AvailableMaterials;
 
@@ -1286,7 +1287,7 @@ void FScene::AddGeometryInstanceFromComponent(ULandscapeComponent* InComponent)
 	ENQUEUE_RENDER_COMMAND(RenderThreadInit)(
 		[
 			InstanceRenderState = MoveTemp(InstanceRenderState),
-			FeatureLevel,
+			LocalFeatureLevel = FeatureLevel,
 			Initializer,
 			InstanceLightmapRenderStateInitializers = MoveTemp(InstanceLightmapRenderStateInitializers),
 			&RenderState = RenderState,
@@ -1300,11 +1301,11 @@ void FScene::AddGeometryInstanceFromComponent(ULandscapeComponent* InComponent)
 		{
 			InstanceRenderState.SharedBuffers = new FLandscapeSharedBuffers(
 				InstanceRenderState.SharedBuffersKey, Initializer.SubsectionSizeQuads, Initializer.NumSubsections,
-				FeatureLevel);
+				LocalFeatureLevel);
 
 			FLandscapeComponentSceneProxy::SharedBuffersMap.Add(InstanceRenderState.SharedBuffersKey, InstanceRenderState.SharedBuffers);
 
-			FLandscapeFixedGridVertexFactory* LandscapeVertexFactory = new FLandscapeFixedGridVertexFactory(FeatureLevel);
+			FLandscapeFixedGridVertexFactory* LandscapeVertexFactory = new FLandscapeFixedGridVertexFactory(LocalFeatureLevel);
 			LandscapeVertexFactory->Data.PositionComponent = FVertexStreamComponent(InstanceRenderState.SharedBuffers->VertexBuffer, 0, sizeof(FLandscapeVertex), VET_Float4);
 			LandscapeVertexFactory->InitResource();
 			InstanceRenderState.SharedBuffers->FixedGridVertexFactory = LandscapeVertexFactory;
@@ -1456,7 +1457,7 @@ void FScene::AddGeometryInstanceFromComponent(ULandscapeComponent* InComponent)
 
 	for (auto ResourceCluster : ResourceClusters)
 	{
-		ResourceCluster->UpdateUniformBuffer(ERHIFeatureLevel::SM5);
+		ResourceCluster->UpdateUniformBuffer(FeatureLevel);
 	}
 
 	if (InComponent->GetWorld())
@@ -1827,11 +1828,11 @@ void GatherBuildDataResourcesToKeep(const ULevel* InLevel, ULevel* LightingScena
 
 void FScene::ApplyFinishedLightmapsToWorld()
 {
-	static const auto CVar = IConsoleManager::Get().FindTConsoleVariableDataInt(TEXT("r.VirtualTexturedLightmaps"));
-	const bool bUseVirtualTextures = (CVar->GetValueOnAnyThread() != 0) && UseVirtualTexturing(GMaxRHIFeatureLevel);
-
 	UWorld* World = GPULightmass->World;
 
+	static const auto CVar = IConsoleManager::Get().FindTConsoleVariableDataInt(TEXT("r.VirtualTexturedLightmaps"));
+	const bool bUseVirtualTextures = (CVar->GetValueOnAnyThread() != 0) && UseVirtualTexturing(World->FeatureLevel);
+
 	{
 		FScopedSlowTask SlowTask(3);
 		SlowTask.MakeDialog();
diff --git a/Engine/Plugins/Experimental/GPULightmass/Source/GPULightmass/Private/Scene/Scene.h b/Engine/Plugins/Experimental/GPULightmass/Source/GPULightmass/Private/Scene/Scene.h
--- a/Engine/Plugins/Experimental/GPULightmass/Source/GPULightmass/Private/Scene/Scene.h
+++ b/Engine/Plugins/Experimental/GPULightmass/Source/GPULightmass/Private/Scene/Scene.h
@@ -117,6 +117,8 @@ public:
 	TUniquePtr<FVolumetricLightmapRenderer> VolumetricLightmapRenderer;
 	TUniquePtr<FIrradianceCache> IrradianceCache;
 
+	ERHIFeatureLevel::Type FeatureLevel;
+
 	int32 GetPrimitiveIdForGPUScene(const FGeometryInstanceRenderStateRef& GeometryInstanceRef) const;
 
 	bool SetupRayTracingScene(int32 LODIndex = 0);
@@ -191,6 +193,8 @@ public:
 
 	FSceneRenderState RenderState;
 
+	ERHIFeatureLevel::Type FeatureLevel;
+
 	bool bNeedsVoxelization = true;
 
 private:
diff --git a/Engine/Plugins/Experimental/GPULightmass/Source/GPULightmass/Private/VolumetricLightmap.cpp b/Engine/Plugins/Experimental/GPULightmass/Source/GPULightmass/Private/VolumetricLightmap.cpp
--- a/Engine/Plugins/Experimental/GPULightmass/Source/GPULightmass/Private/VolumetricLightmap.cpp
+++ b/Engine/Plugins/Experimental/GPULightmass/Source/GPULightmass/Private/VolumetricLightmap.cpp
@@ -100,6 +100,8 @@ END_SHADER_PARAMETER_STRUCT()
 
 void FVolumetricLightmapRenderer::VoxelizeScene()
 {
+	FGlobalShaderMap* GlobalShaderMap = GetGlobalShaderMap(Scene->FeatureLevel);
+
 	for (int32 MipLevel = 0; MipLevel < VoxelizationVolumeMips.Num(); MipLevel++)
 	{
 		VoxelizationVolumeMips[MipLevel].SafeRelease();
@@ -196,8 +198,6 @@ void FVolumetricLightmapRenderer::VoxelizeScene()
 
 	for (int32 MipLevel = 0; MipLevel < VoxelizationVolumeMips.Num(); MipLevel++)
 	{
-		FGlobalShaderMap* GlobalShaderMap = GetGlobalShaderMap(GMaxRHIFeatureLevel);
-
 		FClearVolumeCS::FParameters* Parameters = GraphBuilder.AllocParameters<FClearVolumeCS::FParameters>();
 		Parameters->VolumeSize = VoxelizationVolumeMips[MipLevel]->GetDesc().GetSize();
 		Parameters->VoxelizeVolume = VoxelizationVolumeMips[MipLevel]->GetRenderTargetItem().UAV;
@@ -213,8 +213,6 @@ void FVolumetricLightmapRenderer::VoxelizeScene()
 
 	for (FBox ImportanceVolume : ImportanceVolumes)
 	{
-		FGlobalShaderMap* GlobalShaderMap = GetGlobalShaderMap(GMaxRHIFeatureLevel);
-
 		TShaderMapRef<FVoxelizeImportanceVolumeCS> ComputeShader(GlobalShaderMap);
 
 		FVoxelizeImportanceVolumeCS::FParameters* Parameters = GraphBuilder.AllocParameters<FVoxelizeImportanceVolumeCS::FParameters>();
@@ -341,8 +339,6 @@ void FVolumetricLightmapRenderer::VoxelizeScene()
 	});
 
 	{
-		FGlobalShaderMap* GlobalShaderMap = GetGlobalShaderMap(GMaxRHIFeatureLevel);
-
 		TShaderMapRef<FDilateVolumeCS> ComputeShader(GlobalShaderMap);
 
 		FDilateVolumeCS::FParameters* Parameters = GraphBuilder.AllocParameters<FDilateVolumeCS::FParameters>();
@@ -359,7 +355,7 @@ void FVolumetricLightmapRenderer::VoxelizeScene()
 
 	for (int32 MipLevel = 1; MipLevel < VoxelizationVolumeMips.Num(); MipLevel++)
 	{
-		TShaderMapRef<FDownsampleVolumeCS> ComputeShader(GetGlobalShaderMap(GMaxRHIFeatureLevel));
+		TShaderMapRef<FDownsampleVolumeCS> ComputeShader(GlobalShaderMap);
 
 		FDownsampleVolumeCS::FParameters* Parameters = GraphBuilder.AllocParameters<FDownsampleVolumeCS::FParameters>();
 		Parameters->bIsHighestMip = (MipLevel == VoxelizationVolumeMips.Num() - 1) ? 1 : 0;
@@ -383,8 +379,6 @@ void FVolumetricLightmapRenderer::VoxelizeScene()
 
 	for (int32 MipLevel = VoxelizationVolumeMips.Num() - 1; MipLevel >= 0; MipLevel--)
 	{
-		FGlobalShaderMap* GlobalShaderMap = GetGlobalShaderMap(GMaxRHIFeatureLevel);
-
 		TShaderMapRef<FCountNumBricksCS> ComputeShader(GlobalShaderMap);
 
 		FCountNumBricksCS::FParameters* Parameters = GraphBuilder.AllocParameters<FCountNumBricksCS::FParameters>();
@@ -439,8 +433,6 @@ void FVolumetricLightmapRenderer::VoxelizeScene()
 
 	for (int32 MipLevel = VoxelizationVolumeMips.Num() - 1; MipLevel >= 0; MipLevel--)
 	{
-		FGlobalShaderMap* GlobalShaderMap = GetGlobalShaderMap(GMaxRHIFeatureLevel);
-
 		TShaderMapRef<FGatherBrickRequestsCS> ComputeShader(GlobalShaderMap);
 
 		FGatherBrickRequestsCS::FParameters Parameters;
@@ -457,8 +449,6 @@ void FVolumetricLightmapRenderer::VoxelizeScene()
 
 	for (int32 MipLevel = VoxelizationVolumeMips.Num() - 1; MipLevel >= 0; MipLevel--)
 	{
-		FGlobalShaderMap* GlobalShaderMap = GetGlobalShaderMap(GMaxRHIFeatureLevel);
-
 		TShaderMapRef<FSplatVolumeCS> ComputeShader(GlobalShaderMap);
 
 		FSplatVolumeCS::FParameters Parameters;
@@ -511,6 +501,8 @@ void FVolumetricLightmapRenderer::BackgroundTick()
 	{
 		RDG_EVENT_SCOPE(GraphBuilder, "Volumetric Lightmap Path Tracing");
 
+		FGlobalShaderMap* GlobalShaderMap = GetGlobalShaderMap(Scene->FeatureLevel);
+
 		bool bLastFewFramesIdle = !GCurrentLevelEditingViewportClient || !GCurrentLevelEditingViewportClient->IsRealtime();
 
 		int32 NumSamplesThisFrame = !bLastFewFramesIdle ? 1 : 32;
@@ -542,8 +534,6 @@ void FVolumetricLightmapRenderer::BackgroundTick()
 #if RHI_RAYTRACING
 			if (IsRayTracingEnabled())
 			{
-				FGlobalShaderMap* GlobalShaderMap = GetGlobalShaderMap(GMaxRHIFeatureLevel);
-
 				{
 					FVolumetricLightmapPathTracingRGS::FPermutationDomain PermutationVector;
 					PermutationVector.Set<FVolumetricLightmapPathTracingRGS::FUseIrradianceCaching>(Scene->Settings->bUseIrradianceCaching);
@@ -631,8 +621,6 @@ void FVolumetricLightmapRenderer::BackgroundTick()
 				}
 				);
 
-				FGlobalShaderMap* GlobalShaderMap = GetGlobalShaderMap(GMaxRHIFeatureLevel);
-
 				TShaderMapRef<FFinalizeBrickResultsCS> ComputeShader(GlobalShaderMap);
 
 				FFinalizeBrickResultsCS::FParameters* PassParameters = GraphBuilder.AllocParameters<FFinalizeBrickResultsCS::FParameters>();
@@ -679,8 +667,6 @@ void FVolumetricLightmapRenderer::BackgroundTick()
 				// Doing 2 passes no longer makes sense in an amortized setup
 				// for (int32 StitchPass = 0; StitchPass < 2; StitchPass++)
 				{
-					FGlobalShaderMap* GlobalShaderMap = GetGlobalShaderMap(GMaxRHIFeatureLevel);
-
 					TShaderMapRef<FStitchBorderCS> ComputeShader(GlobalShaderMap);
 
 					FStitchBorderCS::FParameters* PassParameters = GraphBuilder.AllocParameters<FStitchBorderCS::FParameters>();
diff --git a/Engine/Plugins/Experimental/GPULightmass/Source/GPULightmass/Private/VolumetricLightmapVoxelization.h b/Engine/Plugins/Experimental/GPULightmass/Source/GPULightmass/Private/VolumetricLightmapVoxelization.h
--- a/Engine/Plugins/Experimental/GPULightmass/Source/GPULightmass/Private/VolumetricLightmapVoxelization.h
+++ b/Engine/Plugins/Experimental/GPULightmass/Source/GPULightmass/Private/VolumetricLightmapVoxelization.h
@@ -99,7 +99,7 @@ class FVLMVoxelizationMeshProcessor : public FMeshPassProcessor
 {
 public:
 	FVLMVoxelizationMeshProcessor(const FScene* InScene, const FSceneView* InView, FMeshPassDrawListContext* InDrawListContext)
-		: FMeshPassProcessor(InScene, GMaxRHIFeatureLevel, InView, InDrawListContext)
+		: FMeshPassProcessor(InScene, InView->GetFeatureLevel(), InView, InDrawListContext)
 		, DrawRenderState(*InView)
 	{
 		DrawRenderState.SetDepthStencilState(TStaticDepthStencilState<false, CF_Always>::GetRHI());
