Increase maximum allowable value for Realtime & Non-Realtime Workload Factor. Both factors clamped to (GISamples - 1) to avoid convergence issues

diff --git a/Engine/Plugins/Experimental/GPULightmass/Source/GPULightmass/Private/GPULightmassSettings.cpp b/Engine/Plugins/Experimental/GPULightmass/Source/GPULightmass/Private/GPULightmassSettings.cpp
--- a/Engine/Plugins/Experimental/GPULightmass/Source/GPULightmass/Private/GPULightmassSettings.cpp
+++ b/Engine/Plugins/Experimental/GPULightmass/Source/GPULightmass/Private/GPULightmassSettings.cpp
@@ -142,8 +142,8 @@ void UGPULightmassSettings::ApplyImmediateSettingsToRunningInstances()
 		if (GPULightmass)
 		{
 			GPULightmass->Settings->bShowProgressBars = bShowProgressBars;
-			GPULightmass->Settings->TilePassesInSlowMode = TilePassesInSlowMode;
-			GPULightmass->Settings->TilePassesInFullSpeedMode = TilePassesInFullSpeedMode;
+			GPULightmass->Settings->TilePassesInSlowMode = FMath::Min(TilePassesInSlowMode, GISamples - 1);
+			GPULightmass->Settings->TilePassesInFullSpeedMode = FMath::Min(TilePassesInFullSpeedMode, GISamples - 1);
 			GPULightmass->Settings->bVisualizeIrradianceCache = bVisualizeIrradianceCache;
 		}
 	}
diff --git a/Engine/Plugins/Experimental/GPULightmass/Source/GPULightmass/Private/LightmapRenderer.cpp b/Engine/Plugins/Experimental/GPULightmass/Source/GPULightmass/Private/LightmapRenderer.cpp
--- a/Engine/Plugins/Experimental/GPULightmass/Source/GPULightmass/Private/LightmapRenderer.cpp
+++ b/Engine/Plugins/Experimental/GPULightmass/Source/GPULightmass/Private/LightmapRenderer.cpp
@@ -2137,6 +2137,7 @@ void FLightmapRenderer::Finalize(FRDGBuilder& GraphBuilder)
 	}
 
 	int32 NumSamplesPerFrame = (bInsideBackgroundTick && bIsViewportNonRealtime) ? Scene->Settings->TilePassesInFullSpeedMode : Scene->Settings->TilePassesInSlowMode;
+	NumSamplesPerFrame = FMath::Max(FMath::Min(NumSamplesPerFrame, Scene->Settings->GISamples - 1), 0);
 
 	{
 		auto& PendingGITileRequests = *GraphBuilder.AllocObject<TArray<FLightmapTileRequest>>(
diff --git a/Engine/Plugins/Experimental/GPULightmass/Source/GPULightmass/Public/GPULightmassSettings.h b/Engine/Plugins/Experimental/GPULightmass/Source/GPULightmass/Public/GPULightmassSettings.h
--- a/Engine/Plugins/Experimental/GPULightmass/Source/GPULightmass/Public/GPULightmassSettings.h
+++ b/Engine/Plugins/Experimental/GPULightmass/Source/GPULightmass/Public/GPULightmassSettings.h
@@ -104,11 +104,11 @@ public:
 
 	// Baking speed multiplier when Realtime is enabled in the viewer.
 	// Setting this too high can cause the editor to become unresponsive with heavy scenes.
-	UPROPERTY(EditAnywhere, BlueprintReadWrite, Category = System, DisplayName = "Realtime Workload Factor", meta = (ClampMin = "1", ClampMax = "64"))
+	UPROPERTY(EditAnywhere, BlueprintReadWrite, Category = System, DisplayName = "Realtime Workload Factor", meta = (ClampMin = "1", ClampMax = "1024"))
 	int32 TilePassesInSlowMode = 1;
 
 	// Baking speed multiplier when Realtime is disabled in the viewer.
-	UPROPERTY(EditAnywhere, BlueprintReadWrite, Category = System, DisplayName = "Non-Realtime Workload Factor", meta = (ClampMin = "1", ClampMax = "64"))
+	UPROPERTY(EditAnywhere, BlueprintReadWrite, Category = System, DisplayName = "Non-Realtime Workload Factor", meta = (ClampMin = "1", ClampMax = "1024"))
 	int32 TilePassesInFullSpeedMode = 8;
 
 	// GPU Lightmass manages a pool for calculations of visible tiles. The pool size should be set based on the size of the
