GPULM: Provide a scratch buffer for building Landscape BLAS.

diff --git a/Engine/Plugins/Experimental/GPULightmass/Source/GPULightmass/Private/LightmapRenderer.cpp b/Engine/Plugins/Experimental/GPULightmass/Source/GPULightmass/Private/LightmapRenderer.cpp
--- a/Engine/Plugins/Experimental/GPULightmass/Source/GPULightmass/Private/LightmapRenderer.cpp
+++ b/Engine/Plugins/Experimental/GPULightmass/Source/GPULightmass/Private/LightmapRenderer.cpp
@@ -944,7 +944,12 @@ bool FSceneRenderState::SetupRayTracingScene(int32 LODIndex)
 								StaticMeshInstanceRenderStates.Elements.Num() + InstanceGroupRenderStates.Elements.Num() + LandscapeIndex
 							);
 
-							DynamicGeometryCollection.DispatchUpdates(RHICmdList, nullptr);
+							const uint32 BLASScratchSize = DynamicGeometryCollection.ComputeScratchBufferSize();
+
+							FRHIResourceCreateInfo ScratchBufferCreateInfo(TEXT("RHILandscapeScratchBuffer"));
+							FBufferRHIRef ScratchBuffer = RHICreateBuffer(BLASScratchSize, BUF_StructuredBuffer | BUF_RayTracingScratch, 0, ERHIAccess::UAVCompute, ScratchBufferCreateInfo);
+
+							DynamicGeometryCollection.DispatchUpdates(RHICmdList, ScratchBuffer);
 
 							// Landscape VF doesn't really use the vertex buffer in HitGroupSystemParameters
 							// We can release after all related RHI cmds get dispatched onto the cmd list
