GPULightmass: Fix light grid being created per sample (re-use the previous sample's grid instead)

Fix RHI validation errors on "TilePositionsBufferForClear" by moving this buffer to RDG

Fix broken IES support from bad merge

Remove unused View parameter from PrepareLightGrid

Fix several shader compilation warnings

diff --git a/Engine/Plugins/Experimental/GPULightmass/Shaders/Private/IrradianceCachingCommon.ush b/Engine/Plugins/Experimental/GPULightmass/Shaders/Private/IrradianceCachingCommon.ush
--- a/Engine/Plugins/Experimental/GPULightmass/Shaders/Private/IrradianceCachingCommon.ush
+++ b/Engine/Plugins/Experimental/GPULightmass/Shaders/Private/IrradianceCachingCommon.ush
@@ -104,7 +104,7 @@ uint EncodeNormalBits(float3 WorldNormal)
 void EmitGeometryHitPoint(FIrradianceCacheRecord HitPoint, uint Size)
 {
 	uint3 Key = EncodeVoxelKey(floor(HitPoint.WorldPosition.xyz / Size));
-	uint NormalDirectionBits = EncodeNormalBits(HitPoint.WorldNormal);
+	uint NormalDirectionBits = EncodeNormalBits(HitPoint.WorldNormal.xyz);
 	
 	uint HashKey = 0;
 	HashKey = MurmurAdd( HashKey, Key.x );
diff --git a/Engine/Plugins/Experimental/GPULightmass/Shaders/Private/LightmapPathTracing.usf b/Engine/Plugins/Experimental/GPULightmass/Shaders/Private/LightmapPathTracing.usf
--- a/Engine/Plugins/Experimental/GPULightmass/Shaders/Private/LightmapPathTracing.usf
+++ b/Engine/Plugins/Experimental/GPULightmass/Shaders/Private/LightmapPathTracing.usf
@@ -811,7 +811,7 @@ void LightmapPathTracingMainRG()
 						if (BatchedTiles[TileIndex].RenderPassIndex < NumRayGuidingTrialSamples)
 						{
 							int2 ClusterPosition = clamp((int2)LaunchIndex - int2(2, 2), int2(0, 0), int2(63, 63)) / TEXEL_CLUSTER_SIZE;
-							float2 JitteredBin = PrimaryRandSample * DIRECTIONAL_BINS_ONE_DIM;// - float2(0.5f, 0.5f) + PrimaryRandSample.zw;
+							float2 JitteredBin = PrimaryRandSample.xy * DIRECTIONAL_BINS_ONE_DIM;// - float2(0.5f, 0.5f) + PrimaryRandSample.zw;
 							int2 PositionInBin = clamp(JitteredBin, float2(0, 0), float2(DIRECTIONAL_BINS_ONE_DIM - 1, DIRECTIONAL_BINS_ONE_DIM - 1));
 							int2 FinalPosition = ClusterPosition * DIRECTIONAL_BINS_ONE_DIM + PositionInBin;
 							float Illuminance = Luminance(RadianceValue) * saturate(dot(RadianceDirection, ShadingNormal));
@@ -1080,8 +1080,8 @@ void VolumetricLightmapPathTracingMainRG()
 	#endif		
 	
 	float3 JitteredCellPosInVLM = (BrickRequests[BrickIndex].xyz * 4 + CellPosInBrick + Jitter - float3(0.5, 0.5, 0.5)) * BrickRequests[BrickIndex].w;	
-	float DetailCellSize = VolumeSize.xyz / IndirectionTextureDim / 4;
-	float3 WorldPosition = VolumeMin + DetailCellSize * JitteredCellPosInVLM; // RT_LWC_TODO
+	float3 DetailCellSize = VolumeSize.xyz / IndirectionTextureDim / 4;
+	float3 WorldPosition = VolumeMin.xyz + DetailCellSize * JitteredCellPosInVLM; // RT_LWC_TODO
 	float3 TranslatedWorldPosition = WorldPosition + LWCHackToFloat(PrimaryView.PreViewTranslation); // RT_LWC_TODO
 	float3 ShadingNormal = float3(0, 0, (FrameNumber % 2 == 0) ? 1 : -1);
 
diff --git a/Engine/Plugins/Experimental/GPULightmass/Source/GPULightmass/Private/LightmapRayTracing.h b/Engine/Plugins/Experimental/GPULightmass/Source/GPULightmass/Private/LightmapRayTracing.h
--- a/Engine/Plugins/Experimental/GPULightmass/Source/GPULightmass/Private/LightmapRayTracing.h
+++ b/Engine/Plugins/Experimental/GPULightmass/Source/GPULightmass/Private/LightmapRayTracing.h
@@ -51,6 +51,7 @@ class FLightmapPathTracingRGS : public FGlobalShader
 		OutEnvironment.SetDefine(TEXT("GPreviewLightmapPhysicalTileSize"), GPreviewLightmapPhysicalTileSize);
 		OutEnvironment.SetDefine(TEXT("LIGHTMAP_PATH_TRACING_MAIN_RG"), 1);
 		OutEnvironment.CompilerFlags.Add(CFLAG_ForceDXC);
+		OutEnvironment.CompilerFlags.Add(CFLAG_WarningsAsErrors);
 	}
 
 	BEGIN_SHADER_PARAMETER_STRUCT(FParameters, )
@@ -107,6 +108,7 @@ class FVolumetricLightmapPathTracingRGS : public FGlobalShader
 	{
 		OutEnvironment.SetDefine(TEXT("GPreviewLightmapPhysicalTileSize"), GPreviewLightmapPhysicalTileSize);
 		OutEnvironment.CompilerFlags.Add(CFLAG_ForceDXC);
+		OutEnvironment.CompilerFlags.Add(CFLAG_WarningsAsErrors);
 	}
 
 	BEGIN_SHADER_PARAMETER_STRUCT(FParameters, )
@@ -163,6 +165,7 @@ class FStationaryLightShadowTracingRGS : public FGlobalShader
 	{
 		OutEnvironment.SetDefine(TEXT("GPreviewLightmapPhysicalTileSize"), GPreviewLightmapPhysicalTileSize);
 		OutEnvironment.CompilerFlags.Add(CFLAG_ForceDXC);
+		OutEnvironment.CompilerFlags.Add(CFLAG_WarningsAsErrors);
 	}
 
 	BEGIN_SHADER_PARAMETER_STRUCT(FParameters, )
@@ -200,6 +203,7 @@ class FFirstBounceRayGuidingCDFBuildCS : public FGlobalShader
 	{
 		OutEnvironment.SetDefine(TEXT("GPreviewLightmapPhysicalTileSize"), GPreviewLightmapPhysicalTileSize);
 		OutEnvironment.CompilerFlags.Add(CFLAG_WaveOperations);
+		OutEnvironment.CompilerFlags.Add(CFLAG_WarningsAsErrors);
 	}
 
 	BEGIN_SHADER_PARAMETER_STRUCT(FParameters, )
diff --git a/Engine/Plugins/Experimental/GPULightmass/Source/GPULightmass/Private/LightmapRenderer.cpp b/Engine/Plugins/Experimental/GPULightmass/Source/GPULightmass/Private/LightmapRenderer.cpp
--- a/Engine/Plugins/Experimental/GPULightmass/Source/GPULightmass/Private/LightmapRenderer.cpp
+++ b/Engine/Plugins/Experimental/GPULightmass/Source/GPULightmass/Private/LightmapRenderer.cpp
@@ -136,7 +136,7 @@ class FMultiTileClearCS : public FGlobalShader
 	BEGIN_SHADER_PARAMETER_STRUCT(FParameters, )
 		SHADER_PARAMETER(int32, NumTiles)
 		SHADER_PARAMETER(int32, TileSize)
-		SHADER_PARAMETER_SRV(Buffer<int2>, TilePositions)
+		SHADER_PARAMETER_RDG_BUFFER_SRV(Buffer<int2>, TilePositions)
 		SHADER_PARAMETER_RDG_TEXTURE_UAV(RWTexture2D<float4>, TilePool)
 	END_SHADER_PARAMETER_STRUCT()
 };
@@ -2151,6 +2151,10 @@ void FLightmapRenderer::Finalize(FRDGBuilder& GraphBuilder)
 				return !Tile.RenderState->IsTileGIConverged(Tile.VirtualCoordinates, NumGISamples) && Tile.RenderState->GeometryInstanceRef.LODIndex == MostCommonLODIndex;
 			}));
 
+#if RHI_RAYTRACING
+		FLightmapPathTracingRGS::FParameters* PreviousPassParameters[MAX_NUM_GPUS] = {};
+#endif
+
 		// Render GI
 		for (int32 SampleIndex = 0; SampleIndex < NumSamplesPerFrame; SampleIndex++)
 		{
@@ -2177,13 +2181,13 @@ void FLightmapRenderer::Finalize(FRDGBuilder& GraphBuilder)
 							TilePositionsToClear.Add(ScratchTilePoolGPU->GetPositionFromLinearAddress(Tile.TileAddressInScratch));
 						}
 
-						FRWBuffer TilePositionsBuffer;
-						TilePositionsBuffer.Initialize(TEXT("TilePositionsBufferForClear"), TilePositionsToClear.GetTypeSize(), TilePositionsToClear.Num(), PF_R32G32_UINT, BUF_None, &TilePositionsToClear);
-
+						FRDGBufferDesc TilePositionsBufferDesc = FRDGBufferDesc::CreateBufferDesc(TilePositionsToClear.GetTypeSize(), TilePositionsToClear.Num());
+						FRDGBufferRef TilePositionsBuffer = CreateVertexBuffer(GraphBuilder, TEXT("TilePositionsBufferForClear"), TilePositionsBufferDesc, TilePositionsToClear.GetData(), TilePositionsToClear.GetResourceDataSize());
+						FRDGBufferSRVRef TilePositionsBufferSRV = GraphBuilder.CreateSRV(FRDGBufferSRVDesc(TilePositionsBuffer, PF_R32_UINT));
 						FMultiTileClearCS::FParameters* Parameters = GraphBuilder.AllocParameters<FMultiTileClearCS::FParameters>();
 						Parameters->NumTiles = TilePositionsToClear.Num();
 						Parameters->TileSize = GPreviewLightmapPhysicalTileSize;
-						Parameters->TilePositions = HoldReference(GraphBuilder, TilePositionsBuffer.SRV);
+						Parameters->TilePositions = TilePositionsBufferSRV;
 						Parameters->TilePool = ScratchTilePoolLayerUAVs[ScratchLayerIndex];
 
 						TShaderMapRef<FMultiTileClearCS> ComputeShader(GlobalShaderMap);
@@ -2345,11 +2349,26 @@ void FLightmapRenderer::Finalize(FRDGBuilder& GraphBuilder)
 									PassParameters->ViewUniformBuffer = Scene->ReferenceView->ViewUniformBuffer;
 									PassParameters->IrradianceCachingParameters = Scene->IrradianceCache->IrradianceCachingParametersUniformBuffer;
 
-									SetupPathTracingLightParameters(Scene->LightSceneRenderState, GraphBuilder, *Scene->ReferenceView, PassParameters);
-
-									// TODO: find a way to share IES atlas with path tracer ...
-									PassParameters->IESTexture = GraphBuilder.RegisterExternalTexture(GSystemTextures.WhiteDummy, TEXT("IESTexture"));
-									PassParameters->IESTextureSampler = TStaticSamplerState<SF_Bilinear, AM_Clamp, AM_Clamp, AM_Clamp>::GetRHI();
+									if (PreviousPassParameters[GPUIndex] == nullptr)
+									{
+										SetupPathTracingLightParameters(Scene->LightSceneRenderState, GraphBuilder, *Scene->ReferenceView, PassParameters);
+										// store the first pass parameters so we don't have to re-create certain resources constantly
+										PreviousPassParameters[GPUIndex] = PassParameters;
+									}
+									else
+									{
+										PassParameters->LightGridParameters = PreviousPassParameters[GPUIndex]->LightGridParameters;
+										PassParameters->SceneLightCount = PreviousPassParameters[GPUIndex]->SceneLightCount;
+										PassParameters->SceneVisibleLightCount = PreviousPassParameters[GPUIndex]->SceneVisibleLightCount;
+										PassParameters->SceneLights = PreviousPassParameters[GPUIndex]->SceneLights;
+										PassParameters->SkylightTexture = PreviousPassParameters[GPUIndex]->SkylightTexture;
+										PassParameters->SkylightTextureSampler = PreviousPassParameters[GPUIndex]->SkylightTextureSampler;
+										PassParameters->SkylightPdf = PreviousPassParameters[GPUIndex]->SkylightPdf;
+										PassParameters->SkylightInvResolution = PreviousPassParameters[GPUIndex]->SkylightInvResolution;
+										PassParameters->SkylightMipCount = PreviousPassParameters[GPUIndex]->SkylightMipCount;
+										PassParameters->IESTexture = PreviousPassParameters[GPUIndex]->IESTexture;
+										PassParameters->IESTextureSampler = PreviousPassParameters[GPUIndex]->IESTextureSampler;
+									}
 
 									PassParameters->SSProfilesTexture = GetSubsurfaceProfileTexture();
 
@@ -2434,16 +2453,14 @@ void FLightmapRenderer::Finalize(FRDGBuilder& GraphBuilder)
 							TilePositionsToClear.Add(ScratchTilePoolGPU->GetPositionFromLinearAddress(Tile.TileAddressInScratch));
 						}
 
-						FRHIResourceCreateInfo CreateInfo(TEXT("TilePositionsBufferForClear"));
-						CreateInfo.ResourceArray = &TilePositionsToClear;
-						CreateInfo.GPUMask = FRHIGPUMask::FromIndex(GPUIndex);
-						FBufferRHIRef TilePositionsBuffer = RHICreateVertexBuffer(TilePositionsToClear.GetTypeSize() * TilePositionsToClear.Num(), BUF_ShaderResource, CreateInfo);
-						FShaderResourceViewRHIRef TilePositionsBufferSRV = RHICreateShaderResourceView(TilePositionsBuffer, TilePositionsToClear.GetTypeSize(), PF_R32G32_UINT);
+						FRDGBufferDesc TilePositionsBufferDesc = FRDGBufferDesc::CreateBufferDesc(TilePositionsToClear.GetTypeSize(), TilePositionsToClear.Num());
+						FRDGBufferRef TilePositionsBuffer = CreateVertexBuffer(GraphBuilder, TEXT("TilePositionsBufferForClear"), TilePositionsBufferDesc, TilePositionsToClear.GetData(), TilePositionsToClear.GetResourceDataSize());
+						FRDGBufferSRVRef TilePositionsBufferSRV = GraphBuilder.CreateSRV(FRDGBufferSRVDesc(TilePositionsBuffer, PF_R32_UINT));
 
 						FMultiTileClearCS::FParameters* Parameters = GraphBuilder.AllocParameters<FMultiTileClearCS::FParameters>();
 						Parameters->NumTiles = TilePositionsToClear.Num();
 						Parameters->TileSize = GPreviewLightmapPhysicalTileSize;
-						Parameters->TilePositions = HoldReference(GraphBuilder, TilePositionsBufferSRV);
+						Parameters->TilePositions = TilePositionsBufferSRV;
 						Parameters->TilePool = ScratchTilePoolLayerUAVs[ScratchLayerIndex];
 
 						TShaderMapRef<FMultiTileClearCS> ComputeShader(GlobalShaderMap);
diff --git a/Engine/Plugins/Experimental/GPULightmass/Source/GPULightmass/Private/PathTracingLightParameters.inl b/Engine/Plugins/Experimental/GPULightmass/Source/GPULightmass/Private/PathTracingLightParameters.inl
--- a/Engine/Plugins/Experimental/GPULightmass/Source/GPULightmass/Private/PathTracingLightParameters.inl
+++ b/Engine/Plugins/Experimental/GPULightmass/Source/GPULightmass/Private/PathTracingLightParameters.inl
@@ -9,7 +9,7 @@
 
 RENDERER_API FRDGTexture* PrepareIESAtlas(const TMap<FTexture*, int>& InIESLightProfilesMap, FRDGBuilder& GraphBuilder);
 
-RENDERER_API void PrepareLightGrid(FRDGBuilder& GraphBuilder, const FViewInfo& View, FPathTracingLightGrid* LightGridParameters, const FPathTracingLight* Lights, uint32 NumLights, uint32 NumInfiniteLights, FRDGBufferSRV* LightsSRV);
+RENDERER_API void PrepareLightGrid(FRDGBuilder& GraphBuilder, FPathTracingLightGrid* LightGridParameters, const FPathTracingLight* Lights, uint32 NumLights, uint32 NumInfiniteLights, FRDGBufferSRV* LightsSRV);
 
 template<typename PassParameterType>
 void SetupPathTracingLightParameters(
@@ -226,8 +226,9 @@ void SetupPathTracingLightParameters(
 	{
 		PassParameters->IESTexture = GraphBuilder.RegisterExternalTexture(GSystemTextures.WhiteDummy);
 	}
+	PassParameters->IESTextureSampler = TStaticSamplerState<SF_Bilinear, AM_Clamp, AM_Clamp, AM_Clamp>::GetRHI();
 
-	PrepareLightGrid(GraphBuilder, View, &PassParameters->LightGridParameters, Lights.GetData(), PassParameters->SceneLightCount, NumInfiniteLights, PassParameters->SceneLights);
+	PrepareLightGrid(GraphBuilder, &PassParameters->LightGridParameters, Lights.GetData(), PassParameters->SceneLightCount, NumInfiniteLights, PassParameters->SceneLights);
 }
 
 #endif  // RHI_RAYTRACING
diff --git a/Engine/Plugins/Experimental/GPULightmass/Source/GPULightmass/Private/VolumetricLightmap.cpp b/Engine/Plugins/Experimental/GPULightmass/Source/GPULightmass/Private/VolumetricLightmap.cpp
--- a/Engine/Plugins/Experimental/GPULightmass/Source/GPULightmass/Private/VolumetricLightmap.cpp
+++ b/Engine/Plugins/Experimental/GPULightmass/Source/GPULightmass/Private/VolumetricLightmap.cpp
@@ -375,6 +375,8 @@ void FVolumetricLightmapRenderer::VoxelizeScene()
 		InitialBrickAllocatorParams.Add(0);
 		InitialBrickAllocatorParams.Add(0);
 		BrickAllocatorParameters.Initialize(TEXT("VolumetricLightmapBrickAllocatorParameters"), 4, 2, PF_R32_SINT, BUF_UnorderedAccess | BUF_SourceCopy, &InitialBrickAllocatorParams);
+
+		RHICmdList.Transition(FRHITransitionInfo(BrickAllocatorParameters.UAV, ERHIAccess::Unknown, ERHIAccess::UAVCompute));
 	}
 
 	for (int32 MipLevel = VoxelizationVolumeMips.Num() - 1; MipLevel >= 0; MipLevel--)
@@ -523,6 +525,7 @@ void FVolumetricLightmapRenderer::BackgroundTick()
 		}
 		);
 
+		FVolumetricLightmapPathTracingRGS::FParameters* PreviousPassParameters = nullptr;
 		for (int32 SampleIndex = 0; SampleIndex < NumSamplesThisFrame; SampleIndex++)
 		{
 			int32 MaxBricksPerFrame = FMath::Min(512, NumTotalBricks);
@@ -559,10 +562,25 @@ void FVolumetricLightmapRenderer::BackgroundTick()
 					PassParameters->ViewUniformBuffer = Scene->ReferenceView->ViewUniformBuffer;
 					PassParameters->IrradianceCachingParameters = Scene->IrradianceCache->IrradianceCachingParametersUniformBuffer;
 
-					SetupPathTracingLightParameters(Scene->LightSceneRenderState, GraphBuilder, *Scene->ReferenceView, PassParameters);
-
-
-					PassParameters->IESTextureSampler = TStaticSamplerState<SF_Bilinear, AM_Clamp, AM_Clamp, AM_Clamp>::GetRHI();
+					if (PreviousPassParameters == nullptr)
+					{
+						SetupPathTracingLightParameters(Scene->LightSceneRenderState, GraphBuilder, *Scene->ReferenceView, PassParameters);
+						PreviousPassParameters = PassParameters;
+					}
+					else
+					{
+						PassParameters->LightGridParameters = PreviousPassParameters->LightGridParameters;
+						PassParameters->SceneLightCount = PreviousPassParameters->SceneLightCount;
+						PassParameters->SceneVisibleLightCount = PreviousPassParameters->SceneVisibleLightCount;
+						PassParameters->SceneLights = PreviousPassParameters->SceneLights;
+						PassParameters->SkylightTexture = PreviousPassParameters->SkylightTexture;
+						PassParameters->SkylightTextureSampler = PreviousPassParameters->SkylightTextureSampler;
+						PassParameters->SkylightPdf = PreviousPassParameters->SkylightPdf;
+						PassParameters->SkylightInvResolution = PreviousPassParameters->SkylightInvResolution;
+						PassParameters->SkylightMipCount = PreviousPassParameters->SkylightMipCount;
+						PassParameters->IESTexture = PreviousPassParameters->IESTexture;
+						PassParameters->IESTextureSampler = PreviousPassParameters->IESTextureSampler;
+					}
 
 					PassParameters->SSProfilesTexture = GetSubsurfaceProfileTexture();
 
diff --git a/Engine/Source/Runtime/Renderer/Private/PathTracing.cpp b/Engine/Source/Runtime/Renderer/Private/PathTracing.cpp
--- a/Engine/Source/Runtime/Renderer/Private/PathTracing.cpp
+++ b/Engine/Source/Runtime/Renderer/Private/PathTracing.cpp
@@ -574,7 +574,7 @@ class FPathTracingBuildLightGridCS : public FGlobalShader
 
 	static void ModifyCompilationEnvironment(const FGlobalShaderPermutationParameters& Parameters, FShaderCompilerEnvironment& OutEnvironment)
 	{
-		//OutEnvironment.CompilerFlags.Add(CFLAG_WarningsAsErrors);
+		OutEnvironment.CompilerFlags.Add(CFLAG_WarningsAsErrors);
 		OutEnvironment.CompilerFlags.Add(CFLAG_AllowTypedUAVLoads);
 		OutEnvironment.SetDefine(TEXT("THREADGROUPSIZE_X"), FComputeShaderUtils::kGolden2DGroupSize);
 		OutEnvironment.SetDefine(TEXT("THREADGROUPSIZE_Y"), FComputeShaderUtils::kGolden2DGroupSize);
@@ -1151,7 +1151,7 @@ bool PrepareSkyTexture(FRDGBuilder& GraphBuilder, FScene* Scene, const FViewInfo
 	return true;
 }
 
-RENDERER_API void PrepareLightGrid(FRDGBuilder& GraphBuilder, const FViewInfo& View, FPathTracingLightGrid* LightGridParameters, const FPathTracingLight* Lights, uint32 NumLights, uint32 NumInfiniteLights, FRDGBufferSRV* LightsSRV)
+RENDERER_API void PrepareLightGrid(FRDGBuilder& GraphBuilder, FPathTracingLightGrid* LightGridParameters, const FPathTracingLight* Lights, uint32 NumLights, uint32 NumInfiniteLights, FRDGBufferSRV* LightsSRV)
 {
 	const float Inf = std::numeric_limits<float>::infinity();
 	LightGridParameters->SceneInfiniteLightCount = NumInfiniteLights;
@@ -1521,9 +1521,9 @@ void SetLightParameters(FRDGBuilder& GraphBuilder, FPathTracingRG::FParameters*
 	{
 		PassParameters->IESTexture = GraphBuilder.RegisterExternalTexture(GSystemTextures.WhiteDummy);
 	}
+	PassParameters->IESTextureSampler = TStaticSamplerState<SF_Bilinear, AM_Clamp, AM_Clamp, AM_Clamp>::GetRHI();
 
-
-	PrepareLightGrid(GraphBuilder, View, &PassParameters->LightGridParameters, Lights, NumLights, NumInfiniteLights, PassParameters->SceneLights);
+	PrepareLightGrid(GraphBuilder, &PassParameters->LightGridParameters, Lights, NumLights, NumInfiniteLights, PassParameters->SceneLights);
 }
 
 class FPathTracingCompositorPS : public FGlobalShader
@@ -1861,7 +1861,6 @@ void FDeferredShadingSceneRenderer::RenderPathTracing(
 						PassParameters->SceneVisibleLightCount = 0;
 					}
 
-					PassParameters->IESTextureSampler = TStaticSamplerState<SF_Bilinear, AM_Clamp, AM_Clamp, AM_Clamp>::GetRHI();
 					PassParameters->RadianceTexture = GraphBuilder.CreateUAV(RadianceTexture);
 					PassParameters->AlbedoTexture   = GraphBuilder.CreateUAV(AlbedoTexture);
 					PassParameters->NormalTexture   = GraphBuilder.CreateUAV(NormalTexture);
