GPULightmass should complete shader compilation before building lighting.

- The alternative is to have lightmass constantly restart as new shaders complete.

diff --git a/Engine/Plugins/Experimental/GPULightmass/Source/GPULightmass/Private/GPULightmassSettings.cpp b/Engine/Plugins/Experimental/GPULightmass/Source/GPULightmass/Private/GPULightmassSettings.cpp
--- a/Engine/Plugins/Experimental/GPULightmass/Source/GPULightmass/Private/GPULightmassSettings.cpp
+++ b/Engine/Plugins/Experimental/GPULightmass/Source/GPULightmass/Private/GPULightmassSettings.cpp
@@ -9,6 +9,7 @@
 #include "GPULightmass.h"
 #include "Editor.h"
 #include "LevelEditorViewport.h"
+#include "AssetCompilingManager.h"
 
 #define LOCTEXT_NAMESPACE "StaticLightingSystem"
 
@@ -293,6 +294,9 @@ void UGPULightmassSubsystem::Launch()
 		SlowTask.MakeDialog();
 		SlowTask.EnterProgressFrame(1, LOCTEXT("StartingStaticLightingSystem", "Starting static lighting system"));
 
+		UMaterialInterface::SubmitRemainingJobsForWorld(World);
+		FAssetCompilingManager::Get().FinishAllCompilation();
+
 		{
 			FGlobalComponentRecreateRenderStateContext RecreateRenderStateContext; // Implicit FlushRenderingCommands();
 
