GPULM: Fix type error introduced when TilePositionsBuffer was converted to RDG

diff --git a/Engine/Plugins/Experimental/GPULightmass/Source/GPULightmass/Private/LightmapRenderer.cpp b/Engine/Plugins/Experimental/GPULightmass/Source/GPULightmass/Private/LightmapRenderer.cpp
--- a/Engine/Plugins/Experimental/GPULightmass/Source/GPULightmass/Private/LightmapRenderer.cpp
+++ b/Engine/Plugins/Experimental/GPULightmass/Source/GPULightmass/Private/LightmapRenderer.cpp
@@ -2183,7 +2183,7 @@ void FLightmapRenderer::Finalize(FRDGBuilder& GraphBuilder)
 
 						FRDGBufferDesc TilePositionsBufferDesc = FRDGBufferDesc::CreateBufferDesc(TilePositionsToClear.GetTypeSize(), TilePositionsToClear.Num());
 						FRDGBufferRef TilePositionsBuffer = CreateVertexBuffer(GraphBuilder, TEXT("TilePositionsBufferForClear"), TilePositionsBufferDesc, TilePositionsToClear.GetData(), TilePositionsToClear.GetResourceDataSize());
-						FRDGBufferSRVRef TilePositionsBufferSRV = GraphBuilder.CreateSRV(FRDGBufferSRVDesc(TilePositionsBuffer, PF_R32_UINT));
+						FRDGBufferSRVRef TilePositionsBufferSRV = GraphBuilder.CreateSRV(FRDGBufferSRVDesc(TilePositionsBuffer, PF_R32G32_UINT));
 						FMultiTileClearCS::FParameters* Parameters = GraphBuilder.AllocParameters<FMultiTileClearCS::FParameters>();
 						Parameters->NumTiles = TilePositionsToClear.Num();
 						Parameters->TileSize = GPreviewLightmapPhysicalTileSize;
