Deprecated GetRenderTargetItem() from IPooledRenderTarget.

diff --git a/Engine/Plugins/Experimental/GPULightmass/Source/GPULightmass/Private/LightmapRenderer.cpp b/Engine/Plugins/Experimental/GPULightmass/Source/GPULightmass/Private/LightmapRenderer.cpp
--- a/Engine/Plugins/Experimental/GPULightmass/Source/GPULightmass/Private/LightmapRenderer.cpp
+++ b/Engine/Plugins/Experimental/GPULightmass/Source/GPULightmass/Private/LightmapRenderer.cpp
@@ -1449,9 +1449,9 @@ void FLightmapRenderer::Finalize(FRDGBuilder& GraphBuilder)
 			{
 				uint32 DstRowPitch;
 				FLinearColor* Texture[3];
-				Texture[0] = (FLinearColor*)RHICmdList.LockTexture2D(UploadTilePoolGPU->PooledRenderTargets[0]->GetRenderTargetItem().TargetableTexture->GetTexture2D(), 0, RLM_WriteOnly, DstRowPitch, false);
-				Texture[1] = (FLinearColor*)RHICmdList.LockTexture2D(UploadTilePoolGPU->PooledRenderTargets[1]->GetRenderTargetItem().TargetableTexture->GetTexture2D(), 0, RLM_WriteOnly, DstRowPitch, false);
-				Texture[2] = (FLinearColor*)RHICmdList.LockTexture2D(UploadTilePoolGPU->PooledRenderTargets[2]->GetRenderTargetItem().TargetableTexture->GetTexture2D(), 0, RLM_WriteOnly, DstRowPitch, false);
+				Texture[0] = (FLinearColor*)RHICmdList.LockTexture2D(UploadTilePoolGPU->PooledRenderTargets[0]->GetRHI(), 0, RLM_WriteOnly, DstRowPitch, false);
+				Texture[1] = (FLinearColor*)RHICmdList.LockTexture2D(UploadTilePoolGPU->PooledRenderTargets[1]->GetRHI(), 0, RLM_WriteOnly, DstRowPitch, false);
+				Texture[2] = (FLinearColor*)RHICmdList.LockTexture2D(UploadTilePoolGPU->PooledRenderTargets[2]->GetRHI(), 0, RLM_WriteOnly, DstRowPitch, false);
 
 				TSet<FVirtualTile> TilesToDecompress;
 
@@ -1551,9 +1551,9 @@ void FLightmapRenderer::Finalize(FRDGBuilder& GraphBuilder)
 					}
 				});
 
-				RHICmdList.UnlockTexture2D(UploadTilePoolGPU->PooledRenderTargets[0]->GetRenderTargetItem().TargetableTexture->GetTexture2D(), 0, false);
-				RHICmdList.UnlockTexture2D(UploadTilePoolGPU->PooledRenderTargets[1]->GetRenderTargetItem().TargetableTexture->GetTexture2D(), 0, false);
-				RHICmdList.UnlockTexture2D(UploadTilePoolGPU->PooledRenderTargets[2]->GetRenderTargetItem().TargetableTexture->GetTexture2D(), 0, false);
+				RHICmdList.UnlockTexture2D(UploadTilePoolGPU->PooledRenderTargets[0]->GetRHI(), 0, false);
+				RHICmdList.UnlockTexture2D(UploadTilePoolGPU->PooledRenderTargets[1]->GetRHI(), 0, false);
+				RHICmdList.UnlockTexture2D(UploadTilePoolGPU->PooledRenderTargets[2]->GetRHI(), 0, false);
 			}
 
 			FGPUBatchedTileRequests GPUBatchedTileRequests;
@@ -2776,7 +2776,7 @@ void FLightmapRenderer::Finalize(FRDGBuilder& GraphBuilder)
 					FIntRect GPURect;
 					GPURect.Min = LightmapTilePoolGPU.GetPositionFromLinearAddress(Tile.TileAddressInWorkingSet) * LightmapTilePoolGPU.LayerFormatAndTileSize[RenderTargetIndex].TileSize;
 					GPURect.Max = GPURect.Min + LightmapTilePoolGPU.LayerFormatAndTileSize[RenderTargetIndex].TileSize;
-					Params.Add(FTransferResourceParams(LightmapTilePoolGPU.PooledRenderTargets[RenderTargetIndex]->GetRenderTargetItem().TargetableTexture->GetTexture2D(), GPURect, AssignedGPUIndex, 0, true, true));
+					Params.Add(FTransferResourceParams(LightmapTilePoolGPU.PooledRenderTargets[RenderTargetIndex]->GetRHI(), GPURect, AssignedGPUIndex, 0, true, true));
 				};
 
 				TransferTexture(0);
diff --git a/Engine/Plugins/Experimental/GPULightmass/Source/GPULightmass/Private/VolumetricLightmap.cpp b/Engine/Plugins/Experimental/GPULightmass/Source/GPULightmass/Private/VolumetricLightmap.cpp
--- a/Engine/Plugins/Experimental/GPULightmass/Source/GPULightmass/Private/VolumetricLightmap.cpp
+++ b/Engine/Plugins/Experimental/GPULightmass/Source/GPULightmass/Private/VolumetricLightmap.cpp
@@ -174,7 +174,7 @@ void FVolumetricLightmapRenderer::VoxelizeScene()
 	}
 
 	VolumetricLightmapData.Bounds = FinalVolume;
-	VolumetricLightmapData.IndirectionTexture.Texture = IndirectionTexture->GetRenderTargetItem().ShaderResourceTexture->GetTexture3D();
+	VolumetricLightmapData.IndirectionTexture.Texture = IndirectionTexture->GetRHI();
 	VolumetricLightmapData.IndirectionTexture.Format = PF_R8G8B8A8_UINT;
 	VolumetricLightmapData.IndirectionTextureDimensions = FIntVector(IndirectionTextureDimensions);
 	VolumetricLightmapData.BrickSize = 4;
@@ -192,15 +192,15 @@ void FVolumetricLightmapRenderer::VoxelizeScene()
 	VLMVoxelizationParams->VolumeCenter = (FVector3f)CubeVolume.GetCenter(); // LWC_TODO: precision loss
 	VLMVoxelizationParams->VolumeExtent = (FVector3f)CubeVolume.GetExtent(); // LWC_TODO: precision loss
 	VLMVoxelizationParams->VolumeMaxDim = CubeMaxDim;
-	VLMVoxelizationParams->VoxelizeVolume = VoxelizationVolumeMips[0]->GetRenderTargetItem().UAV;
-	VLMVoxelizationParams->IndirectionTexture = IndirectionTexture->GetRenderTargetItem().UAV;
+	VLMVoxelizationParams->VoxelizeVolume = VoxelizationVolumeMips[0]->GetUAV();
+	VLMVoxelizationParams->IndirectionTexture = IndirectionTexture->GetUAV();
 	TRDGUniformBufferRef<FVLMVoxelizationParams> PassUniformBuffer = GraphBuilder.CreateUniformBuffer(VLMVoxelizationParams);
 
 	for (int32 MipLevel = 0; MipLevel < VoxelizationVolumeMips.Num(); MipLevel++)
 	{
 		FClearVolumeCS::FParameters* Parameters = GraphBuilder.AllocParameters<FClearVolumeCS::FParameters>();
 		Parameters->VolumeSize = VoxelizationVolumeMips[MipLevel]->GetDesc().GetSize();
-		Parameters->VoxelizeVolume = VoxelizationVolumeMips[MipLevel]->GetRenderTargetItem().UAV;
+		Parameters->VoxelizeVolume = VoxelizationVolumeMips[MipLevel]->GetUAV();
 
 		TShaderMapRef<FClearVolumeCS> ComputeShader(GlobalShaderMap);
 		FComputeShaderUtils::AddPass(
@@ -220,7 +220,7 @@ void FVolumetricLightmapRenderer::VoxelizeScene()
 		Parameters->ImportanceVolumeMin = (FVector3f)ImportanceVolume.Min;
 		Parameters->ImportanceVolumeMax = (FVector3f)ImportanceVolume.Max;
 		Parameters->VLMVoxelizationParams = PassUniformBuffer;
-		Parameters->VoxelizeVolume = VoxelizationVolumeMips[0]->GetRenderTargetItem().UAV;
+		Parameters->VoxelizeVolume = VoxelizationVolumeMips[0]->GetUAV();
 
 		FComputeShaderUtils::AddPass(
 			GraphBuilder,
@@ -343,7 +343,7 @@ void FVolumetricLightmapRenderer::VoxelizeScene()
 
 		FDilateVolumeCS::FParameters* Parameters = GraphBuilder.AllocParameters<FDilateVolumeCS::FParameters>();
 		Parameters->VolumeSize = VoxelizationVolumeMips[0]->GetDesc().GetSize();
-		Parameters->VoxelizeVolume = VoxelizationVolumeMips[0]->GetRenderTargetItem().UAV;
+		Parameters->VoxelizeVolume = VoxelizationVolumeMips[0]->GetUAV();
 
 		FComputeShaderUtils::AddPass(
 			GraphBuilder,
@@ -359,8 +359,8 @@ void FVolumetricLightmapRenderer::VoxelizeScene()
 
 		FDownsampleVolumeCS::FParameters* Parameters = GraphBuilder.AllocParameters<FDownsampleVolumeCS::FParameters>();
 		Parameters->bIsHighestMip = (MipLevel == VoxelizationVolumeMips.Num() - 1) ? 1 : 0;
-		Parameters->VoxelizeVolume = VoxelizationVolumeMips[MipLevel]->GetRenderTargetItem().UAV;
-		Parameters->VoxelizeVolumePrevMip = VoxelizationVolumeMips[MipLevel - 1]->GetRenderTargetItem().UAV;
+		Parameters->VoxelizeVolume = VoxelizationVolumeMips[MipLevel]->GetUAV();
+		Parameters->VoxelizeVolumePrevMip = VoxelizationVolumeMips[MipLevel - 1]->GetUAV();
 
 		FComputeShaderUtils::AddPass(
 			GraphBuilder,
@@ -385,7 +385,7 @@ void FVolumetricLightmapRenderer::VoxelizeScene()
 
 		FCountNumBricksCS::FParameters* Parameters = GraphBuilder.AllocParameters<FCountNumBricksCS::FParameters>();
 		Parameters->VolumeSize = VoxelizationVolumeMips[MipLevel]->GetDesc().GetSize();
-		Parameters->VoxelizeVolume = VoxelizationVolumeMips[MipLevel]->GetRenderTargetItem().UAV;
+		Parameters->VoxelizeVolume = VoxelizationVolumeMips[MipLevel]->GetUAV();
 		Parameters->BrickAllocatorParameters = BrickAllocatorParameters.UAV;
 
 		FComputeShaderUtils::AddPass(
@@ -440,7 +440,7 @@ void FVolumetricLightmapRenderer::VoxelizeScene()
 		FGatherBrickRequestsCS::FParameters Parameters;
 		Parameters.VolumeSize = VoxelizationVolumeMips[MipLevel]->GetDesc().GetSize();
 		Parameters.BrickSize = 1 << (MipLevel * BrickSizeLog2);
-		Parameters.VoxelizeVolume = VoxelizationVolumeMips[MipLevel]->GetRenderTargetItem().UAV;
+		Parameters.VoxelizeVolume = VoxelizationVolumeMips[MipLevel]->GetUAV();
 		Parameters.BrickAllocatorParameters = BrickAllocatorParameters.UAV;
 		Parameters.BrickRequests = BrickRequests.UAV;
 
@@ -457,13 +457,13 @@ void FVolumetricLightmapRenderer::VoxelizeScene()
 		Parameters.VolumeSize = IndirectionTextureDimensions;
 		Parameters.BrickSize = 1 << (MipLevel * BrickSizeLog2);
 		Parameters.bIsHighestMip = MipLevel == VoxelizationVolumeMips.Num() - 1;
-		Parameters.VoxelizeVolume = VoxelizationVolumeMips[MipLevel]->GetRenderTargetItem().UAV;
-		Parameters.IndirectionTexture = IndirectionTexture->GetRenderTargetItem().UAV;
+		Parameters.VoxelizeVolume = VoxelizationVolumeMips[MipLevel]->GetUAV();
+		Parameters.IndirectionTexture = IndirectionTexture->GetUAV();
 		Parameters.BrickAllocatorParameters = BrickAllocatorParameters.UAV;
 
 		FComputeShaderUtils::Dispatch(RHICmdList, ComputeShader, Parameters, FComputeShaderUtils::GetGroupCount(IndirectionTextureDimensions, FIntVector(4)));
 
-		RHICmdList.Transition(FRHITransitionInfo(IndirectionTexture->GetRenderTargetItem().UAV, ERHIAccess::UAVCompute, ERHIAccess::UAVCompute));
+		RHICmdList.Transition(FRHITransitionInfo(IndirectionTexture->GetUAV(), ERHIAccess::UAVCompute, ERHIAccess::UAVCompute));
 	}
 
 	Scene->DestroyRayTracingScene();
@@ -693,7 +693,7 @@ void FVolumetricLightmapRenderer::BackgroundTick()
 					PassParameters->FrameNumber = FrameNumber / NumFramesOneRound;
 					PassParameters->NumTotalBricks = NumTotalBricks;
 					PassParameters->BrickBatchOffset = BrickBatchOffset;
-					PassParameters->IndirectionTexture = IndirectionTexture->GetRenderTargetItem().UAV;
+					PassParameters->IndirectionTexture = IndirectionTexture->GetUAV();
 					PassParameters->BrickRequests = BrickRequests.UAV;
 					PassParameters->AmbientVector = AccumulationBrickData.AmbientVector.Texture;
 					PassParameters->OutAmbientVector = VolumetricLightmapData.BrickData.AmbientVector.UAV;
diff --git a/Engine/Plugins/Experimental/VirtualHeightfieldMesh/Source/VirtualHeightfieldMeshEditor/Private/HeightfieldMinMaxTextureBuild.cpp b/Engine/Plugins/Experimental/VirtualHeightfieldMesh/Source/VirtualHeightfieldMeshEditor/Private/HeightfieldMinMaxTextureBuild.cpp
--- a/Engine/Plugins/Experimental/VirtualHeightfieldMesh/Source/VirtualHeightfieldMeshEditor/Private/HeightfieldMinMaxTextureBuild.cpp
+++ b/Engine/Plugins/Experimental/VirtualHeightfieldMesh/Source/VirtualHeightfieldMeshEditor/Private/HeightfieldMinMaxTextureBuild.cpp
@@ -207,7 +207,7 @@ namespace VirtualHeightfieldMesh
 					Desc.bIsThumbnails = false;
 					Desc.DebugType = ERuntimeVirtualTextureDebugType::None;
 					Desc.NumPageDescs = 1;
-					Desc.Targets[0].Texture = RenderTileResources.GetTileRenderTarget()->GetRenderTargetItem().ShaderResourceTexture->GetTexture2D();
+					Desc.Targets[0].Texture = RenderTileResources.GetTileRenderTarget()->GetRHI();
 					Desc.Targets[1].Texture = nullptr;
 					Desc.Targets[2].Texture = nullptr;
 					Desc.PageDescs[0].DestBox[0] = TileBox;
@@ -255,7 +255,7 @@ namespace VirtualHeightfieldMesh
 					CopyInfo.DestMipIndex = 0;
 
 					RHICmdList.Transition(FRHITransitionInfo(RenderTileResources.GetStagingTexture(MipLevel), ERHIAccess::Unknown, ERHIAccess::CopyDest));
-					RHICmdList.CopyTexture(RenderTileResources.GetFinalRenderTarget()->GetRenderTargetItem().ShaderResourceTexture, RenderTileResources.GetStagingTexture(MipLevel), CopyInfo);
+					RHICmdList.CopyTexture(RenderTileResources.GetFinalRenderTarget()->GetRHI(), RenderTileResources.GetStagingTexture(MipLevel), CopyInfo);
 				}
 
 				RHICmdList.WriteGPUFence(RenderTileResources.GetFence());
diff --git a/Engine/Plugins/FX/Niagara/Source/Niagara/Private/DataInterface/NiagaraDistanceFieldParameters.h b/Engine/Plugins/FX/Niagara/Source/Niagara/Private/DataInterface/NiagaraDistanceFieldParameters.h
--- a/Engine/Plugins/FX/Niagara/Source/Niagara/Private/DataInterface/NiagaraDistanceFieldParameters.h
+++ b/Engine/Plugins/FX/Niagara/Source/Niagara/Private/DataInterface/NiagaraDistanceFieldParameters.h
@@ -64,7 +64,7 @@ public:
 			SetShaderValue(RHICmdList, ShaderRHI, NumSceneObjects, ParameterData->NumObjectsInBuffer);
 			SetSRVParameter(RHICmdList, ShaderRHI, SceneDistanceFieldAssetData, ParameterData->AssetDataBuffer.SRV);
 			SetSRVParameter(RHICmdList, ShaderRHI, DistanceFieldIndirectionTable, ParameterData->IndirectionTable.SRV);
-			SetTextureParameter(RHICmdList, ShaderRHI, DistanceFieldBrickTexture, ParameterData->DistanceFieldBrickVolumeTexture->GetRenderTargetItem().ShaderResourceTexture);
+			SetTextureParameter(RHICmdList, ShaderRHI, DistanceFieldBrickTexture, ParameterData->DistanceFieldBrickVolumeTexture->GetRHI());
 			SetSamplerParameter(RHICmdList, ShaderRHI, DistanceFieldSampler, TStaticSamplerState<SF_Bilinear, AM_Clamp, AM_Clamp, AM_Clamp>::GetRHI());
 			SetShaderValue(RHICmdList, ShaderRHI, DistanceFieldBrickSize, FVector3f(DistanceField::BrickSize));
 			SetShaderValue(RHICmdList, ShaderRHI, DistanceFieldUniqueDataBrickSize, FVector3f(DistanceField::UniqueDataBrickSize));
diff --git a/Engine/Plugins/Media/AvfMedia/Source/AvfMedia/Private/Player/AvfMediaVideoSampler.cpp b/Engine/Plugins/Media/AvfMedia/Source/AvfMedia/Private/Player/AvfMediaVideoSampler.cpp
--- a/Engine/Plugins/Media/AvfMedia/Source/AvfMedia/Private/Player/AvfMediaVideoSampler.cpp
+++ b/Engine/Plugins/Media/AvfMedia/Source/AvfMedia/Private/Player/AvfMediaVideoSampler.cpp
@@ -320,7 +320,7 @@ void FAvfMediaVideoSampler::ProcessFrame(CVPixelBufferRef Frame, FTimespan Sampl
 					RHICmdList.DrawPrimitive(0, 2, 1);
 				}
 				RHICmdList.EndRenderPass();
-				RHICmdList.CopyToResolveTarget(ShaderResource, ShaderResource, FResolveParams());
+				RHICmdList.Transition(FRHITransitionInfo(ShaderResource, ERHIAccess::RTV, ERHIAccess::SRVMask));
 			}
 			
 			CFRelease(YTextureRef);
diff --git a/Engine/Plugins/Media/ElectraUtil/Source/ElectraSamples/Private/ElectraTextureSampleAVF.cpp b/Engine/Plugins/Media/ElectraUtil/Source/ElectraSamples/Private/ElectraTextureSampleAVF.cpp
--- a/Engine/Plugins/Media/ElectraUtil/Source/ElectraSamples/Private/ElectraTextureSampleAVF.cpp
+++ b/Engine/Plugins/Media/ElectraUtil/Source/ElectraSamples/Private/ElectraTextureSampleAVF.cpp
@@ -255,7 +255,7 @@ void FElectraMediaTexConvApple::ConvertTexture(FTexture2DRHIRef & InDstTexture,
 					RHICmdList.DrawPrimitive(0, 2, 1);
 				}
 				RHICmdList.EndRenderPass();
-				RHICmdList.CopyToResolveTarget(ShaderResource, ShaderResource, FResolveParams());
+				RHICmdList.Transition(FRHITransitionInfo(ShaderResource, ERHIAccess::RTV, ERHIAccess::SRVMask));
 			}
 			CFRelease(YTextureRef);
 			CFRelease(UVTextureRef);
diff --git a/Engine/Plugins/Media/MediaIOFramework/Source/MediaIOCore/Private/MediaCapture.cpp b/Engine/Plugins/Media/MediaIOFramework/Source/MediaIOCore/Private/MediaCapture.cpp
--- a/Engine/Plugins/Media/MediaIOFramework/Source/MediaIOCore/Private/MediaCapture.cpp
+++ b/Engine/Plugins/Media/MediaIOFramework/Source/MediaIOCore/Private/MediaCapture.cpp
@@ -827,9 +827,9 @@ void UMediaCapture::Capture_RenderThread(FRHICommandListImmediate& RHICmdList,
 		else
 		{
 			TRACE_CPUPROFILER_EVENT_SCOPE(UMediaCapture::BeforeFrameCaptured);
-			UE_LOG(LogMediaIOCore, VeryVerbose, TEXT("Before Locking texture %p"), reinterpret_cast<uintptr_t>(CapturingFrame->RenderTarget->GetRenderTargetItem().TargetableTexture->GetTexture2D()->GetNativeResource()));
-			BeforeFrameCaptured_RenderingThread(CapturingFrame->CaptureBaseData, CapturingFrame->UserData, CapturingFrame->RenderTarget->GetRenderTargetItem().TargetableTexture);
-			UE_LOG(LogMediaIOCore, VeryVerbose, TEXT("After Locking texture %p"), reinterpret_cast<uintptr_t>(CapturingFrame->RenderTarget->GetRenderTargetItem().TargetableTexture->GetTexture2D()->GetNativeResource()));
+			UE_LOG(LogMediaIOCore, VeryVerbose, TEXT("Before Locking texture %p"), reinterpret_cast<uintptr_t>(CapturingFrame->RenderTarget->GetRHI()->GetNativeResource()));
+			BeforeFrameCaptured_RenderingThread(CapturingFrame->CaptureBaseData, CapturingFrame->UserData, CapturingFrame->RenderTarget->GetRHI());
+			UE_LOG(LogMediaIOCore, VeryVerbose, TEXT("After Locking texture %p"), reinterpret_cast<uintptr_t>(CapturingFrame->RenderTarget->GetRHI()->GetNativeResource()));
 		}
 	}
 
@@ -915,7 +915,7 @@ void UMediaCapture::Capture_RenderThread(FRHICommandListImmediate& RHICmdList,
 		TRACE_CPUPROFILER_EVENT_SCOPE(UMediaCapture::CopyToResolve);
 		SCOPE_CYCLE_COUNTER(STAT_MediaCapture_RenderThread_CopyToResolve);
 
-		FSceneRenderTargetItem& DestRenderTarget = CapturingFrame->RenderTarget->GetRenderTargetItem();
+		FRHITexture* DestRenderTarget = CapturingFrame->RenderTarget->GetRHI();
 		// Do we need to crop
 		float ULeft = 0.0f;
 		float URight = 1.0f;
@@ -961,19 +961,19 @@ void UMediaCapture::Capture_RenderThread(FRHICommandListImmediate& RHICmdList,
 			if (InMediaCapture->ConversionOperation == EMediaCaptureConversionOperation::NONE && !bRequiresFormatConversion)
 			{
 				// Asynchronously copy target from GPU to GPU
-				RHICmdList.CopyToResolveTarget(SourceTexture, DestRenderTarget.TargetableTexture, ResolveParams);
+				RHICmdList.CopyToResolveTarget(SourceTexture, DestRenderTarget, ResolveParams);
 			}
 			else if (InMediaCapture->ConversionOperation == EMediaCaptureConversionOperation::CUSTOM)
 			{
 				InMediaCapture->OnCustomCapture_RenderingThread(RHICmdList, CapturingFrame->CaptureBaseData, CapturingFrame->UserData
-					, SourceTexture, DestRenderTarget.TargetableTexture, ResolveParams, {ULeft, URight}, {VTop, VBottom});
+					, SourceTexture, DestRenderTarget, ResolveParams, {ULeft, URight}, {VTop, VBottom});
 			}
 			else
 			{
 				TRACE_CPUPROFILER_EVENT_SCOPE(UMediaCapture::FormatConversion);
 				// convert the source with a draw call
 				FGraphicsPipelineStateInitializer GraphicsPSOInit;
-				FRHITexture* RenderTarget = DestRenderTarget.TargetableTexture.GetReference();
+				FRHITexture* RenderTarget = DestRenderTarget;
 				RHICmdList.Transition(FRHITransitionInfo(RenderTarget, ERHIAccess::Unknown, ERHIAccess::RTV));
 				FRHIRenderPassInfo RPInfo(RenderTarget,  ERenderTargetActions::DontLoad_Store);
 				RHICmdList.BeginRenderPass(RPInfo, TEXT("MediaCapture"));
@@ -1041,19 +1041,19 @@ void UMediaCapture::Capture_RenderThread(FRHICommandListImmediate& RHICmdList,
 				RHICmdList.SetViewport(0, 0, 0.0f, InMediaCapture->DesiredOutputSize.X, InMediaCapture->DesiredOutputSize.Y, 1.0f);
 				RHICmdList.DrawPrimitive(0, 2, 1);
 				RHICmdList.EndRenderPass();
-				RHICmdList.Transition(FRHITransitionInfo(DestRenderTarget.TargetableTexture, ERHIAccess::RTV, ERHIAccess::SRVGraphics));
+				RHICmdList.Transition(FRHITransitionInfo(DestRenderTarget, ERHIAccess::RTV, ERHIAccess::SRVGraphics));
 
 			}
 
 			if (InMediaCapture->bShouldCaptureRHITexture)
 			{
-				InMediaCapture->OnRHITextureCaptured_RenderingThread(CapturingFrame->CaptureBaseData, CapturingFrame->UserData, DestRenderTarget.TargetableTexture);
+				InMediaCapture->OnRHITextureCaptured_RenderingThread(CapturingFrame->CaptureBaseData, CapturingFrame->UserData, DestRenderTarget);
 				CapturingFrame->bResolvedTargetRequested = false;
 			}
 			else
 			{
 				// Asynchronously copy duplicate target from GPU to System Memory
-				RHICmdList.CopyToResolveTarget(DestRenderTarget.TargetableTexture, CapturingFrame->ReadbackTexture, FResolveParams());
+				RHICmdList.CopyToResolveTarget(DestRenderTarget, CapturingFrame->ReadbackTexture, FResolveParams());
 				CapturingFrame->bResolvedTargetRequested = true;
 				const int32 FrameIndex = CaptureFrames.IndexOfByPredicate([&CapturingFrame](const FCaptureFrame& InCaptureFrame) { return &InCaptureFrame == CapturingFrame; });
 				UE_LOG(LogMediaIOCore, VeryVerbose, TEXT("Requested copy for capture frame %d"), FrameIndex);
diff --git a/Engine/Plugins/MovieScene/MovieRenderPipeline/Source/MovieRenderPipelineCore/Private/MoviePipelineSurfaceReader.cpp b/Engine/Plugins/MovieScene/MovieRenderPipeline/Source/MovieRenderPipelineCore/Private/MoviePipelineSurfaceReader.cpp
--- a/Engine/Plugins/MovieScene/MovieRenderPipeline/Source/MovieRenderPipelineCore/Private/MoviePipelineSurfaceReader.cpp
+++ b/Engine/Plugins/MovieScene/MovieRenderPipeline/Source/MovieRenderPipelineCore/Private/MoviePipelineSurfaceReader.cpp
@@ -112,7 +112,7 @@ void FMoviePipelineSurfaceReader::ResolveSampleToReadbackTexture_RenderThread(co
 
 	// Enqueue a render pass which uses a simple shader and a fullscreen quad to copy the SourceSurfaceSample to this SurfaceReader's cpu-readback enabled texture.
 	// This RenderPass resolves to our ReadbackTexture 
-	FRHIRenderPassInfo RPInfo(ResampleTexturePooledRenderTarget->GetRenderTargetItem().TargetableTexture, ERenderTargetActions::Load_Store, ReadbackTexture);
+	FRHIRenderPassInfo RPInfo(ResampleTexturePooledRenderTarget->GetRHI(), ERenderTargetActions::Load_Store, ReadbackTexture);
 	RHICmdList.BeginRenderPass(RPInfo, TEXT("MoviePipelineSurfaceResolveRenderTarget"));
 	{
 		RHICmdList.SetViewport(0, 0, 0.0f, TargetSize.X, TargetSize.Y, 1.0f);
diff --git a/Engine/Plugins/Runtime/TextureShare/Source/TextureShare/Private/TextureShareRHI.cpp b/Engine/Plugins/Runtime/TextureShare/Source/TextureShare/Private/TextureShareRHI.cpp
--- a/Engine/Plugins/Runtime/TextureShare/Source/TextureShare/Private/TextureShareRHI.cpp
+++ b/Engine/Plugins/Runtime/TextureShare/Source/TextureShare/Private/TextureShareRHI.cpp
@@ -168,7 +168,7 @@ bool FTextureShareRHI::WriteToShareTexture_RenderThread(FRHICommandListImmediate
 		TRefCountPtr<IPooledRenderTarget> ResampledRTT;
 		if (GetPooledTempRTT_RenderThread(RHICmdList, DstRect.Size(), DstFormat, true, ResampledRTT))
 		{
-			FTexture2DRHIRef RHIResampledRTT = (const FTexture2DRHIRef&)ResampledRTT->GetRenderTargetItem().TargetableTexture;
+			FTexture2DRHIRef RHIResampledRTT = (const FTexture2DRHIRef&)ResampledRTT->GetRHI();
 			if (RHIResampledRTT.IsValid())
 			{
 				// Resample source texture to PooledTempRTT (Src texture now always shader resource)
@@ -208,8 +208,8 @@ bool FTextureShareRHI::ReadFromShareTexture_RenderThread(FRHICommandListImmediat
 		TRefCountPtr<IPooledRenderTarget> ReceivedSRV, ResampledRTT;
 		if (GetPooledTempRTT_RenderThread(RHICmdList, SrcRect.Size(), SrcFormat, false, ReceivedSRV) && GetPooledTempRTT_RenderThread(RHICmdList, DstRect.Size(), DstFormat, true, ResampledRTT))
 		{
-			FTexture2DRHIRef RHIReceivedSRV  = (const FTexture2DRHIRef&)ReceivedSRV->GetRenderTargetItem().ShaderResourceTexture;
-			FTexture2DRHIRef RHIResampledRTT = (const FTexture2DRHIRef&)ResampledRTT->GetRenderTargetItem().TargetableTexture;
+			FTexture2DRHIRef RHIReceivedSRV  = (const FTexture2DRHIRef&)ReceivedSRV->GetRHI();
+			FTexture2DRHIRef RHIResampledRTT = (const FTexture2DRHIRef&)ResampledRTT->GetRHI();
 			if (RHIReceivedSRV.IsValid() && RHIResampledRTT.IsValid())
 			{
 				// Copy direct from shared texture to RHIReceivedSRV (received shared texture has only flag TexCreate_ResolveTargetable, not shader resource)
diff --git a/Engine/Shaders/Private/TemporalSuperResolution/TSRFilterFrequencies.usf b/Engine/Shaders/Private/TemporalSuperResolution/TSRFilterFrequencies.usf
--- a/Engine/Shaders/Private/TemporalSuperResolution/TSRFilterFrequencies.usf
+++ b/Engine/Shaders/Private/TemporalSuperResolution/TSRFilterFrequencies.usf
@@ -176,7 +176,7 @@ void MainCS(
 
 
 	// Process the samples
-	tsr_halfC CenterInput;
+	tsr_halfC CenterInput = 0;
 	tsr_halfC FilteredInput;
 	tsr_halfC FilteredHistory;
 	{
diff --git a/Engine/Source/Runtime/Engine/Private/Rendering/SubsurfaceProfile.cpp b/Engine/Source/Runtime/Engine/Private/Rendering/SubsurfaceProfile.cpp
--- a/Engine/Source/Runtime/Engine/Private/Rendering/SubsurfaceProfile.cpp
+++ b/Engine/Source/Runtime/Engine/Private/Rendering/SubsurfaceProfile.cpp
@@ -429,7 +429,7 @@ void FSubsurfaceProfileTexture::CreateTexture(FRHICommandListImmediate& RHICmdLi
 
 	// Write the contents of the texture.
 	uint32 DestStride;
-	uint8* DestBuffer = (uint8*)RHICmdList.LockTexture2D((FTexture2DRHIRef&)GSSProfiles->GetRenderTargetItem().ShaderResourceTexture, 0, RLM_WriteOnly, DestStride, false);
+	uint8* DestBuffer = (uint8*)RHICmdList.LockTexture2D(GSSProfiles->GetRHI(), 0, RLM_WriteOnly, DestStride, false);
 
 	FLinearColor TextureRow[Width];
 	FMemory::Memzero(TextureRow);
@@ -592,7 +592,7 @@ void FSubsurfaceProfileTexture::CreateTexture(FRHICommandListImmediate& RHICmdLi
 		}
 	}
 
-	RHICmdList.UnlockTexture2D((FTexture2DRHIRef&)GSSProfiles->GetRenderTargetItem().ShaderResourceTexture, 0, false);
+	RHICmdList.UnlockTexture2D(GSSProfiles->GetRHI(), 0, false);
 }
 
 TCHAR MiniFontCharFromIndex(uint32 Index)
diff --git a/Engine/Source/Runtime/Engine/Private/Slate/DebugCanvas.cpp b/Engine/Source/Runtime/Engine/Private/Slate/DebugCanvas.cpp
--- a/Engine/Source/Runtime/Engine/Private/Slate/DebugCanvas.cpp
+++ b/Engine/Source/Runtime/Engine/Private/Slate/DebugCanvas.cpp
@@ -30,7 +30,7 @@ public:
 	}
 
 	/** Sets the texture that this target renders to */
-	void SetRenderTargetTexture( FTexture2DRHIRef& InRHIRef )
+	void SetRenderTargetTexture( FRHITexture* InRHIRef )
 	{
 		RenderTargetTextureRHI = InRHIRef;
 	}
@@ -215,7 +215,7 @@ void FDebugCanvasDrawer::InitDebugCanvas(FViewportClient* ViewportClient, UWorld
 		{
 			if (StereoLayers)
 			{
-				IStereoLayers::FLayerDesc StereoLayerDesc = StereoLayers->GetDebugCanvasLayerDesc(LayerTexture->GetRenderTargetItem().ShaderResourceTexture);
+				IStereoLayers::FLayerDesc StereoLayerDesc = StereoLayers->GetDebugCanvasLayerDesc(LayerTexture->GetRHI());
 				StereoLayerDesc.Flags |= !bCanvasRenderedLastFrame ? IStereoLayers::LAYER_FLAG_HIDDEN : 0;
 
 				if (LayerID == INVALID_LAYER_ID && bCanvasRenderedLastFrame)
@@ -243,7 +243,7 @@ void FDebugCanvasDrawer::DrawRenderThread(FRHICommandListImmediate& RHICmdList,
 	QUICK_SCOPE_CYCLE_COUNTER(Stat_DrawDebugCanvas);
 	if( RenderThreadCanvas.IsValid() )
 	{
-		FTexture2DRHIRef* RT = (FTexture2DRHIRef*)InWindowBackBuffer;
+		FRHITexture* RT = InWindowBackBuffer != nullptr ? ((FTexture2DRHIRef*)InWindowBackBuffer)->GetReference() : nullptr;
 		FTextureRHIRef HMDSwapchain = nullptr;	// Attention: RT could point to HMDSwapchain later.
 		if (RenderThreadCanvas->IsUsingInternalTexture())
 		{
@@ -282,17 +282,17 @@ void FDebugCanvasDrawer::DrawRenderThread(FRHICommandListImmediate& RHICmdList,
 					}
 				}
 			}
-			RT = reinterpret_cast<FTexture2DRHIRef*>(HMDSwapchain == nullptr ? &LayerTexture->GetRenderTargetItem().ShaderResourceTexture : &HMDSwapchain);
+			RT = reinterpret_cast<FRHITexture*>(HMDSwapchain == nullptr ? LayerTexture->GetRHI() : HMDSwapchain.GetReference());
 		}
-		RHICmdList.Transition(FRHITransitionInfo(*RT, ERHIAccess::Unknown, ERHIAccess::RTV));
-		RenderTarget->SetRenderTargetTexture(*RT);
+		RHICmdList.Transition(FRHITransitionInfo(RT, ERHIAccess::Unknown, ERHIAccess::RTV));
+		RenderTarget->SetRenderTargetTexture(RT);
 
 		bool bNeedToFlipVertical = RenderThreadCanvas->GetAllowSwitchVerticalAxis();
 		// Flip when rendering to the back buffer
 		RenderThreadCanvas->SetAllowSwitchVerticalAxis(true);
-		if (RenderThreadCanvas->IsScaledToRenderTarget() && IsValidRef(*RT)) 
+		if (RenderThreadCanvas->IsScaledToRenderTarget() && RT) 
 		{
-			RenderThreadCanvas->SetRenderTargetRect( FIntRect(0, 0, (*RT)->GetSizeX(), (*RT)->GetSizeY()) );
+			RenderThreadCanvas->SetRenderTargetRect( FIntRect(0, 0, (RT)->GetSizeX(), (RT)->GetSizeY()) );
 		}
 		else
 		{
diff --git a/Engine/Source/Runtime/Engine/Private/Slate/SceneViewport.cpp b/Engine/Source/Runtime/Engine/Private/Slate/SceneViewport.cpp
--- a/Engine/Source/Runtime/Engine/Private/Slate/SceneViewport.cpp
+++ b/Engine/Source/Runtime/Engine/Private/Slate/SceneViewport.cpp
@@ -1781,8 +1781,8 @@ void FSceneViewport::EndRenderFrame(FRHICommandListImmediate& RHICmdList, bool b
 	if (UseSeparateRenderTarget())
 	{
 		if (BufferedSlateHandles[CurrentBufferedTargetIndex])
-		{			
-			RHICmdList.CopyToResolveTarget(RenderTargetTextureRenderThreadRHI, RenderTargetTextureRenderThreadRHI, FResolveParams());
+		{
+			RHICmdList.Transition(FRHITransitionInfo(RenderTargetTextureRenderThreadRHI, ERHIAccess::Unknown, ERHIAccess::SRVMask));
 		}
 	}
 	else
diff --git a/Engine/Source/Runtime/MovieSceneCapture/Private/FrameGrabber.cpp b/Engine/Source/Runtime/MovieSceneCapture/Private/FrameGrabber.cpp
--- a/Engine/Source/Runtime/MovieSceneCapture/Private/FrameGrabber.cpp
+++ b/Engine/Source/Runtime/MovieSceneCapture/Private/FrameGrabber.cpp
@@ -124,9 +124,9 @@ void FViewportSurfaceReader::ResolveRenderTarget(FViewportSurfaceReader* RenderT
 		GRenderTargetPool.FindFreeElement(RHICmdList, OutputDesc, ResampleTexturePooledRenderTarget, TEXT("ResampleTexture"));
 		check(ResampleTexturePooledRenderTarget);
 
-		const FSceneRenderTargetItem& DestRenderTarget = ResampleTexturePooledRenderTarget->GetRenderTargetItem();
+		FRHITexture* DestRenderTarget = ResampleTexturePooledRenderTarget->GetRHI();
 
-		FRHIRenderPassInfo RPInfo(DestRenderTarget.TargetableTexture, ERenderTargetActions::Load_Store, ReadbackTexture);
+		FRHIRenderPassInfo RPInfo(DestRenderTarget, ERenderTargetActions::Load_Store, ReadbackTexture);
 		RHICmdList.BeginRenderPass(RPInfo, TEXT("FrameGrabberResolveRenderTarget"));
 		{
 			RHICmdList.SetViewport(0, 0, 0.0f, TargetSize.X, TargetSize.Y, 1.0f);
diff --git a/Engine/Source/Runtime/RenderCore/Private/RenderGraphResourcePool.cpp b/Engine/Source/Runtime/RenderCore/Private/RenderGraphResourcePool.cpp
--- a/Engine/Source/Runtime/RenderCore/Private/RenderGraphResourcePool.cpp
+++ b/Engine/Source/Runtime/RenderCore/Private/RenderGraphResourcePool.cpp
@@ -215,8 +215,10 @@ TRefCountPtr<FRDGTransientRenderTarget> FRDGTransientResourceAllocator::Allocate
 	RenderTarget->Desc = Translate(Texture->CreateInfo);
 	RenderTarget->Desc.DebugName = Texture->GetName();
 	RenderTarget->LifetimeState = ERDGTransientResourceLifetimeState::Allocated;
+PRAGMA_DISABLE_DEPRECATION_WARNINGS
 	RenderTarget->GetRenderTargetItem().TargetableTexture = Texture->GetRHI();
 	RenderTarget->GetRenderTargetItem().ShaderResourceTexture = Texture->GetRHI();
+PRAGMA_ENABLE_DEPRECATION_WARNINGS
 	return RenderTarget;
 }
 
diff --git a/Engine/Source/Runtime/RenderCore/Public/ClearQuad.h b/Engine/Source/Runtime/RenderCore/Public/ClearQuad.h
--- a/Engine/Source/Runtime/RenderCore/Public/ClearQuad.h
+++ b/Engine/Source/Runtime/RenderCore/Public/ClearQuad.h
@@ -10,7 +10,6 @@
 class FRHICommandList;
 struct FRWBufferStructured;
 struct FRWBuffer;
-struct FSceneRenderTargetItem;
 class FRHIUnorderedAccessView;
 class FGraphicsPipelineStateInitializer;
 
diff --git a/Engine/Source/Runtime/RenderCore/Public/RendererInterface.h b/Engine/Source/Runtime/RenderCore/Public/RendererInterface.h
--- a/Engine/Source/Runtime/RenderCore/Public/RendererInterface.h
+++ b/Engine/Source/Runtime/RenderCore/Public/RendererInterface.h
@@ -548,9 +548,13 @@ struct IPooledRenderTarget
 	 * @return in bytes
 	 **/
 	virtual uint32 ComputeMemorySize() const = 0;
+
 	/** Get the low level internals (texture/surface) */
+	UE_DEPRECATED(5.1, "GetRenderTargetItem is deprecated. Use GetRHI() and GetUAV() instead.")
 	inline FSceneRenderTargetItem& GetRenderTargetItem() { return RenderTargetItem; }
+
 	/** Get the low level internals (texture/surface) */
+	UE_DEPRECATED(5.1, "GetRenderTargetItem is deprecated. Use GetRHI() and GetUAV() instead.")
 	inline const FSceneRenderTargetItem& GetRenderTargetItem() const { return RenderTargetItem; }
 
 	/** Returns if the render target is tracked by a pool. */
@@ -564,8 +568,9 @@ struct IPooledRenderTarget
 	virtual uint32 Release() = 0;
 	virtual uint32 GetRefCount() const = 0;
 
-	FRHITexture* GetRHI() const { return RenderTargetItem.GetRHI(); }
+	FRHITexture* GetRHI() const { return RenderTargetItem.TargetableTexture; }
 
+	FRHIUnorderedAccessView* GetUAV() const { return RenderTargetItem.UAV; }
 
 	UE_DEPRECATED(5.0, "Use GetRHI instead.")
 	FRHITexture* GetTargetableRHI() const { return GetRHI(); }
diff --git a/Engine/Source/Runtime/Renderer/Private/DistanceFieldLightingPost.cpp b/Engine/Source/Runtime/Renderer/Private/DistanceFieldLightingPost.cpp
--- a/Engine/Source/Runtime/Renderer/Private/DistanceFieldLightingPost.cpp
+++ b/Engine/Source/Runtime/Renderer/Private/DistanceFieldLightingPost.cpp
@@ -500,7 +500,7 @@ void UpdateHistory(
 		DistanceFieldAOHistoryViewRect->Max.Y = View.ViewRect.Size().Y / GAODownsampleFactor;
 
 #if WITH_MGPU && 0 // TODO(RDG)
-		FRHITexture* TexturesToCopyForTemporalEffect[] = { BentNormalHistoryOutput->GetRenderTargetItem().ShaderResourceTexture.GetReference() };
+		FRHITexture* TexturesToCopyForTemporalEffect[] = { BentNormalHistoryOutput->GetRHI() };
 		RHICmdList.BroadcastTemporalEffect(FName(NameForTemporalEffect, View.ViewState->UniqueID), TexturesToCopyForTemporalEffect);
 #endif
 	}
diff --git a/Engine/Source/Runtime/Renderer/Private/DistanceFieldObjectManagement.cpp b/Engine/Source/Runtime/Renderer/Private/DistanceFieldObjectManagement.cpp
--- a/Engine/Source/Runtime/Renderer/Private/DistanceFieldObjectManagement.cpp
+++ b/Engine/Source/Runtime/Renderer/Private/DistanceFieldObjectManagement.cpp
@@ -142,9 +142,9 @@ FDistanceFieldAtlasParameters DistanceField::SetupAtlasParameters(const FDistanc
 		SceneParameters.DistanceFieldIndirectionTable = GDFShadowOffsetDataStructure == 0 ? DistanceFieldSceneData.IndirectionTable.SRV : nullptr;
 		SceneParameters.DistanceFieldIndirection2Table = GDFShadowOffsetDataStructure == 1 ? DistanceFieldSceneData.Indirection2Table.SRV : nullptr;
 		SceneParameters.DistanceFieldIndirectionAtlas = GDFShadowOffsetDataStructure == 2 && DistanceFieldSceneData.IndirectionAtlas ?
-			DistanceFieldSceneData.IndirectionAtlas->GetRenderTargetItem().ShaderResourceTexture : nullptr;
+			DistanceFieldSceneData.IndirectionAtlas->GetRHI() : nullptr;
 
-		SceneParameters.DistanceFieldBrickTexture = DistanceFieldSceneData.DistanceFieldBrickVolumeTexture->GetRenderTargetItem().ShaderResourceTexture;
+		SceneParameters.DistanceFieldBrickTexture = DistanceFieldSceneData.DistanceFieldBrickVolumeTexture->GetRHI();
 		SceneParameters.DistanceFieldSampler = TStaticSamplerState<SF_Bilinear,AM_Clamp,AM_Clamp,AM_Clamp>::GetRHI();
 
 		SceneParameters.DistanceFieldBrickSize = FVector3f(DistanceField::BrickSize);
diff --git a/Engine/Source/Runtime/Renderer/Private/DistanceFieldStreaming.cpp b/Engine/Source/Runtime/Renderer/Private/DistanceFieldStreaming.cpp
--- a/Engine/Source/Runtime/Renderer/Private/DistanceFieldStreaming.cpp
+++ b/Engine/Source/Runtime/Renderer/Private/DistanceFieldStreaming.cpp
@@ -972,7 +972,7 @@ bool FDistanceFieldSceneData::ResizeIndirectionAtlasIfNeeded(FRDGBuilder& GraphB
 
 		if (IndirectionAtlas)
 		{
-			FRHIUnorderedAccessView* NewIndirectionAtlasUAV = NewIndirectionAtlas->GetRenderTargetItem().UAV;
+			FRHIUnorderedAccessView* NewIndirectionAtlasUAV = NewIndirectionAtlas->GetUAV();
 
 			AddPass(
 				GraphBuilder,
@@ -982,7 +982,7 @@ bool FDistanceFieldSceneData::ResizeIndirectionAtlasIfNeeded(FRDGBuilder& GraphB
 				TShaderMapRef<FCopyDistanceFieldIndirectionAtlasCS> ComputeShader(GlobalShaderMap);
 				FCopyDistanceFieldIndirectionAtlasCS::FParameters PassParameters;
 				PassParameters.RWIndirectionAtlas = NewIndirectionAtlasUAV;
-				PassParameters.DistanceFieldIndirectionAtlas = OldIndirectionAtlas->GetRenderTargetItem().ShaderResourceTexture;
+				PassParameters.DistanceFieldIndirectionAtlas = OldIndirectionAtlas->GetRHI();
 
 				RHICmdList.Transition(FRHITransitionInfo(NewIndirectionAtlasUAV, ERHIAccess::Unknown, ERHIAccess::UAVCompute));
 				RHICmdList.BeginUAVOverlap(NewIndirectionAtlasUAV);
@@ -1538,7 +1538,7 @@ void FDistanceFieldSceneData::UpdateDistanceFieldAtlas(
 						IndirectionAtlasUpload.Unlock();
 
 						TShaderMapRef<FScatterUploadDistanceFieldIndirectionAtlasCS> ComputeShader(GlobalShaderMap);
-						FRHIUnorderedAccessView* IndirectionAtlasUAV = IndirectionAtlas->GetRenderTargetItem().UAV;
+						FRHIUnorderedAccessView* IndirectionAtlasUAV = IndirectionAtlas->GetUAV();
 						FScatterUploadDistanceFieldIndirectionAtlasCS::FParameters PassParameters;
 						PassParameters.RWIndirectionAtlas = IndirectionAtlasUAV;
 						PassParameters.IndirectionUploadIndices = IndirectionAtlasUpload.IndirectionUploadIndicesBuffer.SRV;
diff --git a/Engine/Source/Runtime/Renderer/Private/GPUBenchmark.cpp b/Engine/Source/Runtime/Renderer/Private/GPUBenchmark.cpp
--- a/Engine/Source/Runtime/Renderer/Private/GPUBenchmark.cpp
+++ b/Engine/Source/Runtime/Renderer/Private/GPUBenchmark.cpp
@@ -66,7 +66,7 @@ public:
 
 		FGlobalShader::SetParameters<FViewUniformShaderParameters>(RHICmdList, ShaderRHI, View.ViewUniformBuffer);
 
-		SetTextureParameter(RHICmdList, ShaderRHI, InputTexture, InputTextureSampler, TStaticSamplerState<>::GetRHI(), Src->GetRenderTargetItem().ShaderResourceTexture);
+		SetTextureParameter(RHICmdList, ShaderRHI, InputTexture, InputTextureSampler, TStaticSamplerState<>::GetRHI(), Src->GetRHI());
 	}
 
 	static const TCHAR* GetSourceFilename()
@@ -559,14 +559,14 @@ void RendererGPUBenchmark(FRHICommandListImmediate& RHICmdList, FSynthBenchmarkR
 				// decide how much work we do in this pass
 				LocalWorkScale[Iteration] = (Iteration / 10.f + 1.f) * WorkScale;
 
-				FRHIRenderPassInfo RPInfo(RTItems[DestRTIndex]->GetRenderTargetItem().TargetableTexture, ERenderTargetActions::Load_Store);
+				FRHIRenderPassInfo RPInfo(RTItems[DestRTIndex]->GetRHI(), ERenderTargetActions::Load_Store);
 				TransitionRenderPassTargets(RHICmdList, RPInfo);
 				RHICmdList.BeginRenderPass(RPInfo, TEXT("GPUBenchmark"));
 				{
 					RunBenchmarkShader(RHICmdList, VertexBuffer, View, MethodId, RTItems[SrcRTIndex], LocalWorkScale[Iteration]);
 				}
 				RHICmdList.EndRenderPass();
-				RHICmdList.CopyToResolveTarget(RTItems[DestRTIndex]->GetRenderTargetItem().TargetableTexture, RTItems[DestRTIndex]->GetRenderTargetItem().ShaderResourceTexture, FResolveParams());
+				RHICmdList.Transition(FRHITransitionInfo(RTItems[DestRTIndex]->GetRHI(), ERHIAccess::RTV, ERHIAccess::SRVMask));
 
 				RHICmdList.EndRenderQuery(TimerQueries[QueryIndex].GetQuery());
 
diff --git a/Engine/Source/Runtime/Renderer/Private/GlobalDistanceField.cpp b/Engine/Source/Runtime/Renderer/Private/GlobalDistanceField.cpp
--- a/Engine/Source/Runtime/Renderer/Private/GlobalDistanceField.cpp
+++ b/Engine/Source/Runtime/Renderer/Private/GlobalDistanceField.cpp
@@ -435,29 +435,29 @@ void FGlobalDistanceFieldInfo::UpdateParameterData(float MaxOcclusionDistance, b
 	{
 		if (PageAtlasTexture)
 		{
-			ParameterData.PageAtlasTexture = PageAtlasTexture->GetRenderTargetItem().ShaderResourceTexture;
+			ParameterData.PageAtlasTexture = PageAtlasTexture->GetRHI();
 		}
 
 		if (CoverageAtlasTexture)
 		{
-			ParameterData.CoverageAtlasTexture = CoverageAtlasTexture->GetRenderTargetItem().ShaderResourceTexture;
+			ParameterData.CoverageAtlasTexture = CoverageAtlasTexture->GetRHI();
 		}
 
 		if (PageTableCombinedTexture)
 		{
 			ensureMsgf(GAOGlobalDistanceFieldCacheMostlyStaticSeparately, TEXT("PageTableCombinedTexture should only be allocated when caching mostly static objects separately."));
-			ParameterData.PageTableTexture = PageTableCombinedTexture->GetRenderTargetItem().ShaderResourceTexture;
+			ParameterData.PageTableTexture = PageTableCombinedTexture->GetRHI();
 		}
 		else if (PageTableLayerTextures[GDF_Full])
 		{
 			ensureMsgf(!GAOGlobalDistanceFieldCacheMostlyStaticSeparately, TEXT("PageTableCombinedTexture should be allocated when caching mostly static objects separately."));
-			ParameterData.PageTableTexture = PageTableLayerTextures[GDF_Full]->GetRenderTargetItem().ShaderResourceTexture;
+			ParameterData.PageTableTexture = PageTableLayerTextures[GDF_Full]->GetRHI();
 		}
 
 		FIntVector MipTextureResolution(1, 1, 1);
 		if (MipTexture)
 		{
-			ParameterData.MipTexture = MipTexture->GetRenderTargetItem().ShaderResourceTexture;
+			ParameterData.MipTexture = MipTexture->GetRHI();
 			MipTextureResolution.X = MipTexture->GetDesc().Extent.X;
 			MipTextureResolution.Y = MipTexture->GetDesc().Extent.Y;
 			MipTextureResolution.Z = MipTexture->GetDesc().Depth;
@@ -1216,7 +1216,7 @@ void ReadbackDistanceFieldClipmap(FRHICommandListImmediate& RHICmdList, FGlobalD
 	GDFReadbackRequest = nullptr;
 
 	//FGlobalDistanceFieldClipmap& ClipMap = GlobalDistanceFieldInfo.Clipmaps[0];
-	//FTextureRHIRef SourceTexture = ClipMap.RenderTarget->GetRenderTargetItem().ShaderResourceTexture;
+	//FTextureRHIRef SourceTexture = ClipMap.RenderTarget->GetRHI();
 	//FIntVector Size = SourceTexture->GetSizeXYZ();
 	
 	//RHICmdList.Read3DSurfaceFloatData(SourceTexture, FIntRect(0, 0, Size.X, Size.Y), FIntPoint(0, Size.Z), Readback->ReadbackData);
diff --git a/Engine/Source/Runtime/Renderer/Private/HairStrands/HairStrandsEnvironment.cpp b/Engine/Source/Runtime/Renderer/Private/HairStrands/HairStrandsEnvironment.cpp
--- a/Engine/Source/Runtime/Renderer/Private/HairStrands/HairStrandsEnvironment.cpp
+++ b/Engine/Source/Runtime/Renderer/Private/HairStrands/HairStrandsEnvironment.cpp
@@ -375,7 +375,7 @@ static void AddHairStrandsEnvironmentLightingPassPS(
 	ParametersPS->MultipleScatterSampleCount = FMath::Max(uint32(GHairStrandsSkyLightingSampleCount), 1u);
 	ParametersPS->HairDualScatteringRoughnessOverride = GetHairDualScatteringRoughnessOverride();
 	ParametersPS->TransmissionDensityScaleFactor = FMath::Max(0.f, GHairStrandsTransmissionDensityScaleFactor);
-	ParametersPS->PreIntegratedGF = GSystemTextures.PreintegratedGF->GetRenderTargetItem().ShaderResourceTexture;
+	ParametersPS->PreIntegratedGF = GSystemTextures.PreintegratedGF->GetRHI();
 	ParametersPS->PreIntegratedGFSampler = TStaticSamplerState<SF_Bilinear, AM_Clamp, AM_Clamp, AM_Clamp>::GetRHI();
 	ParametersPS->SceneTextures = SceneTextures;
 	ParametersPS->VirtualVoxel = VirtualVoxelResources.UniformBuffer;
diff --git a/Engine/Source/Runtime/Renderer/Private/HairStrands/HairStrandsLUT.cpp b/Engine/Source/Runtime/Renderer/Private/HairStrands/HairStrandsLUT.cpp
--- a/Engine/Source/Runtime/Renderer/Private/HairStrands/HairStrandsLUT.cpp
+++ b/Engine/Source/Runtime/Renderer/Private/HairStrands/HairStrandsLUT.cpp
@@ -341,7 +341,7 @@ static FRDGTextureRef InternalGetHairLUT(FRDGBuilder& GraphBuilder, const FViewI
 
 	if (Type == FHairLUTType::HairLUTType_DualScattering)
 	{
-		const bool bNeedUpdate = GSystemTextures.HairLUT0.GetReference() == nullptr || GSystemTextures.HairLUT0.GetReference()->GetRenderTargetItem().ShaderResourceTexture->GetSizeXYZ() != FIntVector(GHairLUTIncidentAngleCount, GHairLUTRoughnessCount, GHairLUTAbsorptionCount);
+		const bool bNeedUpdate = GSystemTextures.HairLUT0.GetReference() == nullptr || GSystemTextures.HairLUT0.GetReference()->GetRHI()->GetSizeXYZ() != FIntVector(GHairLUTIncidentAngleCount, GHairLUTRoughnessCount, GHairLUTAbsorptionCount);
 		if (bNeedUpdate)
 		{
 			Out = AddHairLUTPass(GraphBuilder, View, HairLUTType_DualScattering);
@@ -354,7 +354,7 @@ static FRDGTextureRef InternalGetHairLUT(FRDGBuilder& GraphBuilder, const FViewI
 	}
 	else if (Type == FHairLUTType::HairLUTType_MeanEnergy)
 	{
-		const bool bNeedUpdate = GSystemTextures.HairLUT1.GetReference() == nullptr || GSystemTextures.HairLUT1.GetReference()->GetRenderTargetItem().ShaderResourceTexture->GetSizeXYZ() != FIntVector(GHairLUTIncidentAngleCount, GHairLUTRoughnessCount, GHairLUTAbsorptionCount);
+		const bool bNeedUpdate = GSystemTextures.HairLUT1.GetReference() == nullptr || GSystemTextures.HairLUT1.GetReference()->GetRHI()->GetSizeXYZ() != FIntVector(GHairLUTIncidentAngleCount, GHairLUTRoughnessCount, GHairLUTAbsorptionCount);
 		if (bNeedUpdate)
 		{
 			Out = AddHairLUTPass(GraphBuilder, View, HairLUTType_MeanEnergy);
@@ -367,7 +367,7 @@ static FRDGTextureRef InternalGetHairLUT(FRDGBuilder& GraphBuilder, const FViewI
 	}
 	else if (Type == FHairLUTType::HairLUTType_Coverage)
 	{
-		const bool bNeedUpdate = GSystemTextures.HairLUT2.GetReference() == nullptr || GSystemTextures.HairLUT2.GetReference()->GetRenderTargetItem().ShaderResourceTexture->GetSizeXYZ() != FIntVector(FHairCountToCoverageData::HairCount, FHairCountToCoverageData::HairRadiusCount, 1);
+		const bool bNeedUpdate = GSystemTextures.HairLUT2.GetReference() == nullptr || GSystemTextures.HairLUT2.GetReference()->GetRHI()->GetSizeXYZ() != FIntVector(FHairCountToCoverageData::HairCount, FHairCountToCoverageData::HairRadiusCount, 1);
 		if (bNeedUpdate)
 		{
 			Out = AddHairCoverageLUTPass(GraphBuilder, View);
diff --git a/Engine/Source/Runtime/Renderer/Private/IndirectLightRendering.cpp b/Engine/Source/Runtime/Renderer/Private/IndirectLightRendering.cpp
--- a/Engine/Source/Runtime/Renderer/Private/IndirectLightRendering.cpp
+++ b/Engine/Source/Runtime/Renderer/Private/IndirectLightRendering.cpp
@@ -480,7 +480,7 @@ void SetupReflectionUniformParameters(const FViewInfo& View, FReflectionUniformP
 		if (Scene->SkyLight->bRealTimeCaptureEnabled && Scene->ConvolvedSkyRenderTargetReadyIndex >= 0)
 		{
 			// Cannot blend with this capture mode as of today.
-			SkyLightTextureResource = Scene->ConvolvedSkyRenderTarget[Scene->ConvolvedSkyRenderTargetReadyIndex]->GetRenderTargetItem().ShaderResourceTexture;
+			SkyLightTextureResource = Scene->ConvolvedSkyRenderTarget[Scene->ConvolvedSkyRenderTargetReadyIndex]->GetRHI();
 		}
 		else if (SkyLight.ProcessedTexture)
 		{
@@ -532,7 +532,7 @@ void SetupReflectionUniformParameters(const FViewInfo& View, FReflectionUniformP
 	OutParameters.ReflectionCubemap = CubeArrayTexture;
 	OutParameters.ReflectionCubemapSampler = TStaticSamplerState<SF_Trilinear, AM_Clamp, AM_Clamp, AM_Clamp>::GetRHI();
 
-	OutParameters.PreIntegratedGF = GSystemTextures.PreintegratedGF->GetRenderTargetItem().ShaderResourceTexture;
+	OutParameters.PreIntegratedGF = GSystemTextures.PreintegratedGF->GetRHI();
 	OutParameters.PreIntegratedGFSampler = TStaticSamplerState<SF_Bilinear, AM_Clamp, AM_Clamp, AM_Clamp>::GetRHI();
 }
 
@@ -1057,7 +1057,7 @@ void FDeferredShadingSceneRenderer::RenderDiffuseIndirectAndAmbientOcclusion(
 			PassParameters->DiffuseIndirect = DenoiserOutputs;
 			PassParameters->DiffuseIndirectSampler = TStaticSamplerState<SF_Point>::GetRHI();
 
-			PassParameters->PreIntegratedGF = GSystemTextures.PreintegratedGF->GetRenderTargetItem().ShaderResourceTexture;
+			PassParameters->PreIntegratedGF = GSystemTextures.PreintegratedGF->GetRHI();
 			PassParameters->PreIntegratedGFSampler = TStaticSamplerState<SF_Bilinear, AM_Clamp, AM_Clamp, AM_Clamp>::GetRHI();
 
 			PassParameters->AmbientOcclusionTexture = AmbientOcclusionMask;
@@ -1226,7 +1226,7 @@ void FDeferredShadingSceneRenderer::RenderDiffuseIndirectAndAmbientOcclusion(
 		{
 			FAmbientCubemapCompositePS::FParameters* PassParameters = GraphBuilder.AllocParameters<FAmbientCubemapCompositePS::FParameters>();
 			
-			PassParameters->PreIntegratedGF = GSystemTextures.PreintegratedGF->GetRenderTargetItem().ShaderResourceTexture;
+			PassParameters->PreIntegratedGF = GSystemTextures.PreintegratedGF->GetRHI();
 			PassParameters->PreIntegratedGFSampler = TStaticSamplerState<SF_Bilinear, AM_Clamp, AM_Clamp, AM_Clamp>::GetRHI();
 			
 			PassParameters->AmbientOcclusionTexture = AmbientOcclusionMask;
@@ -1362,7 +1362,7 @@ static void AddSkyReflectionPass(
 		}
 		PassParameters->PS.CloudSkyAOSampler = TStaticSamplerState<SF_Bilinear>::GetRHI();
 
-		PassParameters->PS.PreIntegratedGF = GSystemTextures.PreintegratedGF->GetRenderTargetItem().ShaderResourceTexture;
+		PassParameters->PS.PreIntegratedGF = GSystemTextures.PreintegratedGF->GetRHI();
 		PassParameters->PS.PreIntegratedGFSampler = TStaticSamplerState<SF_Bilinear, AM_Clamp, AM_Clamp, AM_Clamp>::GetRHI();
 
 		PassParameters->PS.SceneTextures = SceneTextureParameters;
diff --git a/Engine/Source/Runtime/Renderer/Private/IndirectLightingCache.cpp b/Engine/Source/Runtime/Renderer/Private/IndirectLightingCache.cpp
--- a/Engine/Source/Runtime/Renderer/Private/IndirectLightingCache.cpp
+++ b/Engine/Source/Runtime/Renderer/Private/IndirectLightingCache.cpp
@@ -891,9 +891,9 @@ void FIndirectLightingCache::UpdateBlock(FScene* Scene, FViewInfo* DebugDrawingV
 			BlockInfo.Block.TexelSize);
 
 		// Update the volume texture atlas
-		RHIUpdateTexture3D((const FTexture3DRHIRef&)GetTexture0().ShaderResourceTexture, 0, UpdateRegion, BlockInfo.Block.TexelSize * FormatSize, BlockInfo.Block.TexelSize * BlockInfo.Block.TexelSize * FormatSize, (const uint8*)Texture0Data.GetData());
-		RHIUpdateTexture3D((const FTexture3DRHIRef&)GetTexture1().ShaderResourceTexture, 0, UpdateRegion, BlockInfo.Block.TexelSize * FormatSize, BlockInfo.Block.TexelSize * BlockInfo.Block.TexelSize * FormatSize, (const uint8*)Texture1Data.GetData());
-		RHIUpdateTexture3D((const FTexture3DRHIRef&)GetTexture2().ShaderResourceTexture, 0, UpdateRegion, BlockInfo.Block.TexelSize * FormatSize, BlockInfo.Block.TexelSize * BlockInfo.Block.TexelSize * FormatSize, (const uint8*)Texture2Data.GetData());
+		RHIUpdateTexture3D(GetTexture0(), 0, UpdateRegion, BlockInfo.Block.TexelSize * FormatSize, BlockInfo.Block.TexelSize * BlockInfo.Block.TexelSize * FormatSize, (const uint8*)Texture0Data.GetData());
+		RHIUpdateTexture3D(GetTexture1(), 0, UpdateRegion, BlockInfo.Block.TexelSize * FormatSize, BlockInfo.Block.TexelSize * BlockInfo.Block.TexelSize * FormatSize, (const uint8*)Texture1Data.GetData());
+		RHIUpdateTexture3D(GetTexture2(), 0, UpdateRegion, BlockInfo.Block.TexelSize * FormatSize, BlockInfo.Block.TexelSize * BlockInfo.Block.TexelSize * FormatSize, (const uint8*)Texture2Data.GetData());
 	}
 	else
 	{
diff --git a/Engine/Source/Runtime/Renderer/Private/LightMapRendering.cpp b/Engine/Source/Runtime/Renderer/Private/LightMapRendering.cpp
--- a/Engine/Source/Runtime/Renderer/Private/LightMapRendering.cpp
+++ b/Engine/Source/Runtime/Renderer/Private/LightMapRendering.cpp
@@ -367,9 +367,9 @@ void GetIndirectLightingCacheParameters(
 		// Silently skipping setting the cache texture under failure for now
 		if (FeatureLevel >= ERHIFeatureLevel::SM5 && LightingCache && LightingCache->IsInitialized() && GSupportsVolumeTextureRendering)
 		{
-			Parameters.IndirectLightingCacheTexture0 = const_cast<FIndirectLightingCache*>(LightingCache)->GetTexture0().ShaderResourceTexture;
-			Parameters.IndirectLightingCacheTexture1 = const_cast<FIndirectLightingCache*>(LightingCache)->GetTexture1().ShaderResourceTexture;
-			Parameters.IndirectLightingCacheTexture2 = const_cast<FIndirectLightingCache*>(LightingCache)->GetTexture2().ShaderResourceTexture;
+			Parameters.IndirectLightingCacheTexture0 = const_cast<FIndirectLightingCache*>(LightingCache)->GetTexture0();
+			Parameters.IndirectLightingCacheTexture1 = const_cast<FIndirectLightingCache*>(LightingCache)->GetTexture1();
+			Parameters.IndirectLightingCacheTexture2 = const_cast<FIndirectLightingCache*>(LightingCache)->GetTexture2();
 
 			Parameters.IndirectLightingCacheTextureSampler0 = TStaticSamplerState<SF_Bilinear,AM_Clamp,AM_Clamp,AM_Clamp>::GetRHI();
 			Parameters.IndirectLightingCacheTextureSampler1 = TStaticSamplerState<SF_Bilinear,AM_Clamp,AM_Clamp,AM_Clamp>::GetRHI();
diff --git a/Engine/Source/Runtime/Renderer/Private/MobileBasePassRendering.cpp b/Engine/Source/Runtime/Renderer/Private/MobileBasePassRendering.cpp
--- a/Engine/Source/Runtime/Renderer/Private/MobileBasePassRendering.cpp
+++ b/Engine/Source/Runtime/Renderer/Private/MobileBasePassRendering.cpp
@@ -192,7 +192,7 @@ void SetupMobileBasePassUniformParameters(
 	SetupPlanarReflectionUniformParameters(View, ReflectionSceneProxy, BasePassParameters.PlanarReflection);
 	if (View.PrevViewInfo.MobilePixelProjectedReflection.IsValid())
 	{
-		BasePassParameters.PlanarReflection.PlanarReflectionTexture = View.PrevViewInfo.MobilePixelProjectedReflection->GetRenderTargetItem().ShaderResourceTexture;
+		BasePassParameters.PlanarReflection.PlanarReflectionTexture = View.PrevViewInfo.MobilePixelProjectedReflection->GetRHI();
 		BasePassParameters.PlanarReflection.PlanarReflectionSampler = TStaticSamplerState<SF_Bilinear, AM_Clamp, AM_Clamp, AM_Clamp>::GetRHI();
 	}
 	else if (BasePassParameters.PlanarReflection.PlanarReflectionTexture == nullptr)
@@ -205,7 +205,7 @@ void SetupMobileBasePassUniformParameters(
 
 	SetupMobileSceneTextureUniformParameters(GraphBuilder, SetupMode, BasePassParameters.SceneTextures);
 
-	BasePassParameters.PreIntegratedGFTexture = GSystemTextures.PreintegratedGF->GetRenderTargetItem().ShaderResourceTexture;
+	BasePassParameters.PreIntegratedGFTexture = GSystemTextures.PreintegratedGF->GetRHI();
 	BasePassParameters.PreIntegratedGFSampler = TStaticSamplerState<SF_Bilinear, AM_Clamp, AM_Clamp, AM_Clamp>::GetRHI();
 	BasePassParameters.EyeAdaptationBuffer = GraphBuilder.CreateSRV(GetEyeAdaptationBuffer(GraphBuilder, View), PF_A32B32G32R32F);
 
@@ -289,7 +289,7 @@ void SetupMobileDirectionalLightUniformParameters(
 						const FIntPoint ShadowBufferResolution = ShadowInfo->GetShadowBufferResolution();
 						const FVector4f ShadowBufferSizeValue((float)ShadowBufferResolution.X, (float)ShadowBufferResolution.Y, 1.0f / (float)ShadowBufferResolution.X, 1.0f / (float)ShadowBufferResolution.Y);
 
-						Params.DirectionalLightShadowTexture = ShadowInfo->RenderTargets.DepthTarget->GetRenderTargetItem().ShaderResourceTexture.GetReference();
+						Params.DirectionalLightShadowTexture = ShadowInfo->RenderTargets.DepthTarget->GetRHI();
 						Params.DirectionalLightDirectionAndShadowTransition.W = 1.0f / ShadowInfo->ComputeTransitionSize();
 						Params.DirectionalLightShadowSize = ShadowBufferSizeValue;
 						Params.DirectionalLightDistanceFadeMADAndSpecularScale.W = ShadowInfo->GetShaderReceiverDepthBias();
diff --git a/Engine/Source/Runtime/Renderer/Private/MobileDeferredShadingPass.cpp b/Engine/Source/Runtime/Renderer/Private/MobileDeferredShadingPass.cpp
--- a/Engine/Source/Runtime/Renderer/Private/MobileDeferredShadingPass.cpp
+++ b/Engine/Source/Runtime/Renderer/Private/MobileDeferredShadingPass.cpp
@@ -260,17 +260,17 @@ static void RenderDirectLight(FRHICommandListImmediate& RHICmdList, const FScene
 	SetupReflectionUniformParameters(View, ReflectionUniformParameters);
 	PassParameters.ReflectionsParameters = CreateUniformBufferImmediate(ReflectionUniformParameters, UniformBuffer_SingleDraw);
 	PassParameters.LightFunctionParameters = FVector4f(1.0f, 1.0f, 0.0f, 0.0f);
-	PassParameters.PreIntegratedGF = GSystemTextures.PreintegratedGF->GetRenderTargetItem().ShaderResourceTexture;
+	PassParameters.PreIntegratedGF = GSystemTextures.PreintegratedGF->GetRHI();
 	PassParameters.PreIntegratedGFSampler = TStaticSamplerState<SF_Bilinear, AM_Clamp, AM_Clamp, AM_Clamp>::GetRHI();
 
 	if (IsMobileDistanceFieldEnabled(View.GetShaderPlatform()) && GScreenSpaceShadowMaskTextureMobileOutputs.ScreenSpaceShadowMaskTextureMobile.IsValid())
 	{
-		PassParameters.ScreenSpaceShadowMaskTexture = GScreenSpaceShadowMaskTextureMobileOutputs.ScreenSpaceShadowMaskTextureMobile->GetRenderTargetItem().ShaderResourceTexture;
+		PassParameters.ScreenSpaceShadowMaskTexture = GScreenSpaceShadowMaskTextureMobileOutputs.ScreenSpaceShadowMaskTextureMobile->GetRHI();
 		PassParameters.ScreenSpaceShadowMaskSampler = TStaticSamplerState<SF_Point, AM_Clamp, AM_Clamp, AM_Clamp>::GetRHI();
 	}
 	else
 	{
-		PassParameters.ScreenSpaceShadowMaskTexture = GSystemTextures.WhiteDummy->GetRenderTargetItem().ShaderResourceTexture;
+		PassParameters.ScreenSpaceShadowMaskTexture = GSystemTextures.WhiteDummy->GetRHI();
 		PassParameters.ScreenSpaceShadowMaskSampler = TStaticSamplerState<SF_Point, AM_Clamp, AM_Clamp, AM_Clamp>::GetRHI();
 	}
 	
@@ -469,7 +469,7 @@ static void RenderLocalLight(
 	const FVector InverseScale = FVector(1.f / Scale.Z, 1.f / Scale.Y, 1.f / Scale.X);
 	const FMatrix WorldToLight = LightSceneInfo.Proxy->GetWorldToLight() * FScaleMatrix(FVector(InverseScale));
 	PassParameters.TranslatedWorldToLight = FMatrix44f(FTranslationMatrix(-View.ViewMatrices.GetPreViewTranslation()) * WorldToLight);
-	PassParameters.PreIntegratedGF = GSystemTextures.PreintegratedGF->GetRenderTargetItem().ShaderResourceTexture;
+	PassParameters.PreIntegratedGF = GSystemTextures.PreintegratedGF->GetRHI();
 	PassParameters.PreIntegratedGFSampler = TStaticSamplerState<SF_Bilinear, AM_Clamp, AM_Clamp, AM_Clamp>::GetRHI();
 	FMobileRadialLightFunctionPS::SetParameters(RHICmdList, PixelShader, View, LightMaterial.MaterialProxy, *LightMaterial.Material, PassParameters);
 
diff --git a/Engine/Source/Runtime/Renderer/Private/MobileShadingRenderer.cpp b/Engine/Source/Runtime/Renderer/Private/MobileShadingRenderer.cpp
--- a/Engine/Source/Runtime/Renderer/Private/MobileShadingRenderer.cpp
+++ b/Engine/Source/Runtime/Renderer/Private/MobileShadingRenderer.cpp
@@ -1016,7 +1016,7 @@ void FMobileSceneRenderer::RenderForward(FRDGBuilder& GraphBuilder, FRDGTextureR
 	BasePassRenderTargets.DepthStencil = bIsFullDepthPrepassEnabled ? 
 		FDepthStencilBinding(SceneDepth, ERenderTargetLoadAction::ELoad, ERenderTargetLoadAction::ELoad, FExclusiveDepthStencil::DepthRead_StencilWrite) : 
 		FDepthStencilBinding(SceneDepth, ERenderTargetLoadAction::EClear, ERenderTargetLoadAction::EClear, FExclusiveDepthStencil::DepthWrite_StencilWrite);
-	BasePassRenderTargets.ShadingRateTexture = (!MainView.bIsSceneCapture && !MainView.bIsReflectionCapture && ShadingRateTarget.IsValid()) ? RegisterExternalTexture(GraphBuilder, ShadingRateTarget->GetRenderTargetItem().ShaderResourceTexture, TEXT("ShadingRateTexture")) : nullptr;
+	BasePassRenderTargets.ShadingRateTexture = (!MainView.bIsSceneCapture && !MainView.bIsReflectionCapture && ShadingRateTarget.IsValid()) ? RegisterExternalTexture(GraphBuilder, ShadingRateTarget->GetRHI(), TEXT("ShadingRateTexture")) : nullptr;
 	BasePassRenderTargets.SubpassHint = ESubpassHint::DepthReadSubpass;
 	if (!bIsFullDepthPrepassEnabled)
 	{
@@ -1743,11 +1743,11 @@ void FMobileSceneRenderer::UpdateMovablePointLightUniformBufferAndShadowInfo()
 						FMobileMovableSpotLightsShadowInfo& MobileMovableSpotLightsShadowInfo = View.MobileMovableSpotLightsShadowInfo;
 
 						checkSlow(MobileMovableSpotLightsShadowInfo.ShadowDepthTexture == nullptr
-							|| MobileMovableSpotLightsShadowInfo.ShadowDepthTexture == ProjectedShadowInfo->RenderTargets.DepthTarget->GetRenderTargetItem().ShaderResourceTexture.GetReference());
+							|| MobileMovableSpotLightsShadowInfo.ShadowDepthTexture == ProjectedShadowInfo->RenderTargets.DepthTarget->GetRHI());
 
 						if (MobileMovableSpotLightsShadowInfo.ShadowDepthTexture == nullptr)
 						{
-							MobileMovableSpotLightsShadowInfo.ShadowDepthTexture = ProjectedShadowInfo->RenderTargets.DepthTarget->GetRenderTargetItem().ShaderResourceTexture.GetReference();
+							MobileMovableSpotLightsShadowInfo.ShadowDepthTexture = ProjectedShadowInfo->RenderTargets.DepthTarget->GetRHI();
 							MobileMovableSpotLightsShadowInfo.ShadowBufferSize = FVector4f(ShadowBufferResolution.X, ShadowBufferResolution.Y, 1.0f / ShadowBufferResolution.X, 1.0f / ShadowBufferResolution.Y);
 						}
 					}
diff --git a/Engine/Source/Runtime/Renderer/Private/PostProcess/PostProcessAmbientOcclusionMobile.cpp b/Engine/Source/Runtime/Renderer/Private/PostProcess/PostProcessAmbientOcclusionMobile.cpp
--- a/Engine/Source/Runtime/Renderer/Private/PostProcess/PostProcessAmbientOcclusionMobile.cpp
+++ b/Engine/Source/Runtime/Renderer/Private/PostProcess/PostProcessAmbientOcclusionMobile.cpp
@@ -159,8 +159,8 @@ public:
 
 		if (GSystemTextures.GTAOPreIntegrated.IsValid())
 		{
-			ShaderParameters.GTAOPreIntegrated2D = GSystemTextures.GTAOPreIntegrated->GetRenderTargetItem().ShaderResourceTexture;
-			ShaderParameters.GTAOPreIntegrated3D = GSystemTextures.GTAOPreIntegrated->GetRenderTargetItem().ShaderResourceTexture;
+			ShaderParameters.GTAOPreIntegrated2D = GSystemTextures.GTAOPreIntegrated->GetRHI();
+			ShaderParameters.GTAOPreIntegrated3D = GSystemTextures.GTAOPreIntegrated->GetRHI();
 			ShaderParameters.GTAOPreIntegratedSampler = TStaticSamplerState<SF_Bilinear, AM_Clamp, AM_Clamp, AM_Clamp>::GetRHI();
 		}
 	}
diff --git a/Engine/Source/Runtime/Renderer/Private/PostProcess/PostProcessEyeAdaptation.cpp b/Engine/Source/Runtime/Renderer/Private/PostProcess/PostProcessEyeAdaptation.cpp
--- a/Engine/Source/Runtime/Renderer/Private/PostProcess/PostProcessEyeAdaptation.cpp
+++ b/Engine/Source/Runtime/Renderer/Private/PostProcess/PostProcessEyeAdaptation.cpp
@@ -1012,7 +1012,7 @@ static const FName NAME_EyeAdaptation(TEXT("EyeAdaptation"));
 
 void FSceneViewState::BroadcastEyeAdaptationTemporalEffect(FRHICommandList& RHICmdList)
 {
-	FRHITexture* EyeAdaptation = GetCurrentEyeAdaptationTexture(RHICmdList)->GetRenderTargetItem().ShaderResourceTexture.GetReference();
+	FRHITexture* EyeAdaptation = GetCurrentEyeAdaptationTexture(RHICmdList)->GetRHI();
 	RHICmdList.BroadcastTemporalEffect(FName(NAME_EyeAdaptation, UniqueID), { &EyeAdaptation, 1 });
 }
 
diff --git a/Engine/Source/Runtime/Renderer/Private/RayTracing/RaytracingSkylight.cpp b/Engine/Source/Runtime/Renderer/Private/RayTracing/RaytracingSkylight.cpp
--- a/Engine/Source/Runtime/Renderer/Private/RayTracing/RaytracingSkylight.cpp
+++ b/Engine/Source/Runtime/Renderer/Private/RayTracing/RaytracingSkylight.cpp
@@ -730,8 +730,8 @@ void FDeferredShadingSceneRenderer::VisualizeSkyLightMipTree(
 	TShaderMapRef<FVisualizeSkyLightMipTreePS> PixelShader(ShaderMap);
 	FRHITexture* RenderTargets[2] =
 	{
-		SceneColor->GetRenderTargetItem().TargetableTexture,
-		SkyLightMipTreeRT->GetRenderTargetItem().TargetableTexture
+		SceneColor->GetRHI(),
+		SkyLightMipTreeRT->GetRHI()
 	};
 	FRHIRenderPassInfo RenderPassInfo(2, RenderTargets, ERenderTargetActions::Load_Store);
 	RHICmdList.BeginRenderPass(RenderPassInfo, TEXT("SkyLight Visualization"));
@@ -773,7 +773,7 @@ void FDeferredShadingSceneRenderer::VisualizeSkyLightMipTree(
 	RHICmdList.EndRenderPass();
 	GVisualizeTexture.SetCheckPoint(RHICmdList, SkyLightMipTreeRT);
 
-	RHICmdList.CopyToResolveTarget(SceneColor->GetRenderTargetItem().TargetableTexture, SceneColor->GetRenderTargetItem().ShaderResourceTexture, FResolveParams());
+	RHICmdList.CopyToResolveTarget(SceneColor->GetRHI(), SceneColor->GetRHI(), FResolveParams());
 
 	// Transition to compute
 	RHICmdList.Transition({
diff --git a/Engine/Source/Runtime/Renderer/Private/ReflectionEnvironment.cpp b/Engine/Source/Runtime/Renderer/Private/ReflectionEnvironment.cpp
--- a/Engine/Source/Runtime/Renderer/Private/ReflectionEnvironment.cpp
+++ b/Engine/Source/Runtime/Renderer/Private/ReflectionEnvironment.cpp
@@ -263,7 +263,7 @@ void FReflectionEnvironmentCubemapArray::ResizeCubemapArrayGPU(uint32 InMaxCubem
 
 	InitDynamicRHI();
 
-	FTextureRHIRef TexRef = OldReflectionEnvs->GetRenderTargetItem().TargetableTexture;
+	FTextureRHIRef TexRef = OldReflectionEnvs->GetRHI();
 	FRHICommandListImmediate& RHICmdList = FRHICommandListExecutor::GetImmediateCommandList();
 	const int32 NumMips = FMath::CeilLogTwo(InCubemapSize) + 1;
 
@@ -295,7 +295,7 @@ void FReflectionEnvironmentCubemapArray::ResizeCubemapArrayGPU(uint32 InMaxCubem
 						// For now, we'll just do this on RHIs where we know CopyToResolveTarget does the right thing. In future we should look to 
 						// add a a new RHI method
 						check(GRHISupportsResolveCubemapFaces);
-						RHICmdList.CopyToResolveTarget(OldReflectionEnvs->GetRenderTargetItem().ShaderResourceTexture, ReflectionEnvs->GetRenderTargetItem().ShaderResourceTexture, ResolveParams);
+						RHICmdList.CopyToResolveTarget(OldReflectionEnvs->GetRHI(), ReflectionEnvs->GetRHI(), ResolveParams);
 					}
 				}
 			}
diff --git a/Engine/Source/Runtime/Renderer/Private/ReflectionEnvironmentRealTimeCapture.cpp b/Engine/Source/Runtime/Renderer/Private/ReflectionEnvironmentRealTimeCapture.cpp
--- a/Engine/Source/Runtime/Renderer/Private/ReflectionEnvironmentRealTimeCapture.cpp
+++ b/Engine/Source/Runtime/Renderer/Private/ReflectionEnvironmentRealTimeCapture.cpp
@@ -512,7 +512,7 @@ void FScene::AllocateAndCaptureFrameSkyEnvMap(
 						AtmosphereSetup->ComputeViewData(
 							SkyLight->CapturePosition, MainView.ViewMatrices.GetPreViewTranslation(), SkyViewLutReferentialForward, SkyViewLutReferentialRight,
 							SkyCameraTranslatedWorldOrigin, TempSkyPlanetData, SkyViewLutReferential);
-						CubeView.CachedViewUniformShaderParameters->SkyViewLutTexture = RealTimeReflectionCaptureSkyAtmosphereViewLutTexture->GetRenderTargetItem().ShaderResourceTexture;
+						CubeView.CachedViewUniformShaderParameters->SkyViewLutTexture = RealTimeReflectionCaptureSkyAtmosphereViewLutTexture->GetRHI();
 					}
 					else
 					{
@@ -534,11 +534,11 @@ void FScene::AllocateAndCaptureFrameSkyEnvMap(
 				if (HasSkyAtmosphere()&& (MainView.bSceneHasSkyMaterial || HasVolumetricCloud())
 					&& RealTimeReflectionCaptureCamera360APLutTexture.IsValid())	// we also check that because it seems it can happen for some view setup UE-107270, TODO find a repro for a proper fix.
 				{
-					CubeView.CachedViewUniformShaderParameters->CameraAerialPerspectiveVolume = RealTimeReflectionCaptureCamera360APLutTexture->GetRenderTargetItem().ShaderResourceTexture;
+					CubeView.CachedViewUniformShaderParameters->CameraAerialPerspectiveVolume = RealTimeReflectionCaptureCamera360APLutTexture->GetRHI();
 				}
 				else
 				{
-					CubeView.CachedViewUniformShaderParameters->CameraAerialPerspectiveVolume = GSystemTextures.VolumetricBlackDummy->GetRenderTargetItem().ShaderResourceTexture;
+					CubeView.CachedViewUniformShaderParameters->CameraAerialPerspectiveVolume = GSystemTextures.VolumetricBlackDummy->GetRHI();
 				}
 
 				CubeView.CreateViewUniformBuffers(*CubeView.CachedViewUniformShaderParameters);
diff --git a/Engine/Source/Runtime/Renderer/Private/ScenePrivate.h b/Engine/Source/Runtime/Renderer/Private/ScenePrivate.h
--- a/Engine/Source/Runtime/Renderer/Private/ScenePrivate.h
+++ b/Engine/Source/Runtime/Renderer/Private/ScenePrivate.h
@@ -2244,9 +2244,9 @@ public:
 	void SetLightingCacheDirty(FScene* Scene, const FPrecomputedLightVolume* Volume);
 
 	// Accessors
-	FSceneRenderTargetItem& GetTexture0() { return Texture0->GetRenderTargetItem(); }
-	FSceneRenderTargetItem& GetTexture1() { return Texture1->GetRenderTargetItem(); }
-	FSceneRenderTargetItem& GetTexture2() { return Texture2->GetRenderTargetItem(); }
+	FRHITexture* GetTexture0() { return Texture0->GetRHI(); }
+	FRHITexture* GetTexture1() { return Texture1->GetRHI(); }
+	FRHITexture* GetTexture2() { return Texture2->GetRHI(); }
 
 private:
 	/** Internal helper to determine if indirect lighting is enabled at all */
diff --git a/Engine/Source/Runtime/Renderer/Private/SceneRendering.cpp b/Engine/Source/Runtime/Renderer/Private/SceneRendering.cpp
--- a/Engine/Source/Runtime/Renderer/Private/SceneRendering.cpp
+++ b/Engine/Source/Runtime/Renderer/Private/SceneRendering.cpp
@@ -1135,7 +1135,7 @@ void UpdateNoiseTextureParameters(FViewUniformShaderParameters& ViewUniformShade
 {
 	if (GSystemTextures.PerlinNoiseGradient.GetReference())
 	{
-		ViewUniformShaderParameters.PerlinNoiseGradientTexture = (FTexture2DRHIRef&)GSystemTextures.PerlinNoiseGradient->GetRenderTargetItem().ShaderResourceTexture;
+		ViewUniformShaderParameters.PerlinNoiseGradientTexture = GSystemTextures.PerlinNoiseGradient->GetRHI();
 		SetBlack2DIfNull(ViewUniformShaderParameters.PerlinNoiseGradientTexture);
 	}
 	check(ViewUniformShaderParameters.PerlinNoiseGradientTexture);
@@ -1143,7 +1143,7 @@ void UpdateNoiseTextureParameters(FViewUniformShaderParameters& ViewUniformShade
 
 	if (GSystemTextures.PerlinNoise3D.GetReference())
 	{
-		ViewUniformShaderParameters.PerlinNoise3DTexture = (FTexture3DRHIRef&)GSystemTextures.PerlinNoise3D->GetRenderTargetItem().ShaderResourceTexture;
+		ViewUniformShaderParameters.PerlinNoise3DTexture = GSystemTextures.PerlinNoise3D->GetRHI();
 		SetBlack3DIfNull(ViewUniformShaderParameters.PerlinNoise3DTexture);
 	}
 	check(ViewUniformShaderParameters.PerlinNoise3DTexture);
@@ -1151,7 +1151,7 @@ void UpdateNoiseTextureParameters(FViewUniformShaderParameters& ViewUniformShade
 
 	if (GSystemTextures.SobolSampling.GetReference())
 	{
-		ViewUniformShaderParameters.SobolSamplingTexture = (FTexture2DRHIRef&)GSystemTextures.SobolSampling->GetRenderTargetItem().ShaderResourceTexture;
+		ViewUniformShaderParameters.SobolSamplingTexture = GSystemTextures.SobolSampling->GetRHI();
 		SetBlack2DIfNull(ViewUniformShaderParameters.SobolSamplingTexture);
 	}
 	check(ViewUniformShaderParameters.SobolSamplingTexture);
@@ -1374,24 +1374,24 @@ void FViewInfo::SetupUniformBufferParameters(
 		const TRefCountPtr<IPooledRenderTarget>& PooledTransmittanceLutTexture = SkyAtmosphere->GetTransmittanceLutTexture();
 		if (PooledTransmittanceLutTexture.IsValid())
 		{
-			TransmittanceLutTextureFound = PooledTransmittanceLutTexture->GetRenderTargetItem().ShaderResourceTexture;
+			TransmittanceLutTextureFound = PooledTransmittanceLutTexture->GetRHI();
 		}
 		const TRefCountPtr<IPooledRenderTarget>& PooledDistantSkyLightLutTexture = SkyAtmosphere->GetDistantSkyLightLutTexture();
 		if (PooledDistantSkyLightLutTexture.IsValid())
 		{
-			DistantSkyLightLutTextureFound = PooledDistantSkyLightLutTexture->GetRenderTargetItem().ShaderResourceTexture;
+			DistantSkyLightLutTextureFound = PooledDistantSkyLightLutTexture->GetRHI();
 		}
 
 		if (this->SkyAtmosphereCameraAerialPerspectiveVolume.IsValid())
 		{
-			CameraAerialPerspectiveVolumeFound = this->SkyAtmosphereCameraAerialPerspectiveVolume->GetRenderTargetItem().ShaderResourceTexture;
+			CameraAerialPerspectiveVolumeFound = this->SkyAtmosphereCameraAerialPerspectiveVolume->GetRHI();
 		}
 
 		float SkyViewLutWidth = 1.0f;
 		float SkyViewLutHeight = 1.0f;
 		if (this->SkyAtmosphereViewLutTexture.IsValid())
 		{
-			SkyViewLutTextureFound = this->SkyAtmosphereViewLutTexture->GetRenderTargetItem().ShaderResourceTexture;
+			SkyViewLutTextureFound = this->SkyAtmosphereViewLutTexture->GetRHI();
 			SkyViewLutWidth = float(this->SkyAtmosphereViewLutTexture->GetDesc().GetSize().X);
 			SkyViewLutHeight = float(this->SkyAtmosphereViewLutTexture->GetDesc().GetSize().Y);
 		}
@@ -1934,9 +1934,9 @@ void FViewInfo::SetupUniformBufferParameters(
 	// Hair global resources 
 	SetUpViewHairRenderInfo(*this, ViewUniformShaderParameters.HairRenderInfo, ViewUniformShaderParameters.HairRenderInfoBits, ViewUniformShaderParameters.HairComponents);
 	ViewUniformShaderParameters.HairScatteringLUTTexture = nullptr;
-	if (GSystemTextures.HairLUT0.IsValid() && GSystemTextures.HairLUT0->GetRenderTargetItem().ShaderResourceTexture)
+	if (GSystemTextures.HairLUT0.IsValid() && GSystemTextures.HairLUT0->GetRHI())
 	{
-		ViewUniformShaderParameters.HairScatteringLUTTexture = GSystemTextures.HairLUT0->GetRenderTargetItem().ShaderResourceTexture;
+		ViewUniformShaderParameters.HairScatteringLUTTexture = GSystemTextures.HairLUT0->GetRHI();
 	}
 	ViewUniformShaderParameters.HairScatteringLUTTexture = OrBlack3DIfNull(ViewUniformShaderParameters.HairScatteringLUTTexture);
 	ViewUniformShaderParameters.HairScatteringLUTSampler = TStaticSamplerState<SF_Bilinear, AM_Clamp, AM_Clamp, AM_Clamp>::GetRHI();
diff --git a/Engine/Source/Runtime/Renderer/Private/ScreenPass.cpp b/Engine/Source/Runtime/Renderer/Private/ScreenPass.cpp
--- a/Engine/Source/Runtime/Renderer/Private/ScreenPass.cpp
+++ b/Engine/Source/Runtime/Renderer/Private/ScreenPass.cpp
@@ -13,15 +13,15 @@ RENDERER_API const FScreenTransform FScreenTransform::Identity(FVector2f(1.0f, 1
 RENDERER_API const FScreenTransform FScreenTransform::ScreenPosToViewportUV(FVector2f(0.5f, -0.5f), FVector2f(0.5f, 0.5f));
 RENDERER_API const FScreenTransform FScreenTransform::ViewportUVToScreenPos(FVector2f(2.0f, -2.0f), FVector2f(-1.0f, 1.0f));
 
-const FTextureRHIRef& GetMiniFontTexture()
+FRHITexture* GetMiniFontTexture()
 {
 	if (GSystemTextures.AsciiTexture)
 	{
-		return GSystemTextures.AsciiTexture->GetRenderTargetItem().ShaderResourceTexture;
+		return GSystemTextures.AsciiTexture->GetRHI();
 	}
 	else
 	{
-		return GSystemTextures.WhiteDummy->GetRenderTargetItem().ShaderResourceTexture;
+		return GSystemTextures.WhiteDummy->GetRHI();
 	}
 }
 
diff --git a/Engine/Source/Runtime/Renderer/Private/ScreenPass.h b/Engine/Source/Runtime/Renderer/Private/ScreenPass.h
--- a/Engine/Source/Runtime/Renderer/Private/ScreenPass.h
+++ b/Engine/Source/Runtime/Renderer/Private/ScreenPass.h
@@ -15,7 +15,7 @@
 bool IsHMDHiddenAreaMaskActive();
 
 // Returns the global engine mini font texture.
-const FTextureRHIRef& GetMiniFontTexture();
+FRHITexture* GetMiniFontTexture();
 
 // Creates and returns an RDG texture for the view family output. Returns null if no RHI texture exists.
 FRDGTextureRef RENDERER_API TryCreateViewFamilyTexture(FRDGBuilder& GraphBuilder, const FSceneViewFamily& ViewFamily);
diff --git a/Engine/Source/Runtime/Renderer/Private/ShaderPrint.cpp b/Engine/Source/Runtime/Renderer/Private/ShaderPrint.cpp
--- a/Engine/Source/Runtime/Renderer/Private/ShaderPrint.cpp
+++ b/Engine/Source/Runtime/Renderer/Private/ShaderPrint.cpp
@@ -665,7 +665,7 @@ namespace ShaderPrint
 		// Non graph managed resources
 		FRDGBufferSRVRef ValueBuffer = GraphBuilder.CreateSRV(View.ShaderPrintData.ShaderPrintValueBuffer);
 		FRDGBufferSRVRef StateBuffer = GraphBuilder.CreateSRV(View.ShaderPrintData.ShaderPrintStateBuffer);
-		FTextureRHIRef FontTexture = GSystemTextures.AsciiTexture->GetRenderTargetItem().ShaderResourceTexture;
+		FTextureRHIRef FontTexture = GSystemTextures.AsciiTexture->GetRHI();
 
 		const ERHIFeatureLevel::Type FeatureLevel = View.GetFeatureLevel();
 		FGlobalShaderMap* GlobalShaderMap = GetGlobalShaderMap(FeatureLevel);
diff --git a/Engine/Source/Runtime/Renderer/Private/ShadingEnergyConservation.cpp b/Engine/Source/Runtime/Renderer/Private/ShadingEnergyConservation.cpp
--- a/Engine/Source/Runtime/Renderer/Private/ShadingEnergyConservation.cpp
+++ b/Engine/Source/Runtime/Renderer/Private/ShadingEnergyConservation.cpp
@@ -62,7 +62,7 @@ namespace ShadingEnergyConservationData
 	#include "ShadingEnergyConservationData.h"
 
 	template<typename TDataType>
-	void LockCopyTexture2D(FRHICommandListImmediate& RHICmdList, FTexture2DRHIRef& Texture, const TDataType* InSrcBuffer, uint32 NumComponents)
+	void LockCopyTexture2D(FRHICommandListImmediate& RHICmdList, FRHITexture* Texture, const TDataType* InSrcBuffer, uint32 NumComponents)
 	{
 		const uint8* SrcBuffer = (const uint8*)InSrcBuffer;
 		const uint32 SrcBytesPerPixel = sizeof(TDataType) * NumComponents;
@@ -80,7 +80,7 @@ namespace ShadingEnergyConservationData
 	}
 
 	template<typename TDataType>
-	void LockCopyTexture3D(FRHICommandListImmediate& RHICmdList, FTexture3DRHIRef& Texture, const TDataType* InSrcBuffer, uint32 NumComponents)
+	void LockCopyTexture3D(FRHICommandListImmediate& RHICmdList, FRHITexture* Texture, const TDataType* InSrcBuffer, uint32 NumComponents)
 	{
 		const uint8* SrcBuffer = (const uint8*)InSrcBuffer;
 		const uint32 SrcBytesPerPixel = sizeof(TDataType) * NumComponents;
@@ -330,10 +330,10 @@ void Init(FRDGBuilder& GraphBuilder, FViewInfo& View)
 				View.ViewState->ShadingEnergyConservationData.ClothEnergyTexture    = GRenderTargetPool.FindFreeElement(FRDGTextureDesc::Create2D(FIntPoint(Size, Size),        SpecFormat, FClearValueBinding::None, TexCreate_ShaderResource), TEXT("Shading.ClothSpecEnergy"));
 				View.ViewState->ShadingEnergyConservationData.DiffuseEnergyTexture  = GRenderTargetPool.FindFreeElement(FRDGTextureDesc::Create2D(FIntPoint(Size, Size),        DiffFormat,	FClearValueBinding::None, TexCreate_ShaderResource), TEXT("Shading.DiffuseEnergy"));
 
-				ShadingEnergyConservationData::LockCopyTexture2D(GraphBuilder.RHICmdList, (FTexture2DRHIRef&)View.ViewState->ShadingEnergyConservationData.GGXSpecEnergyTexture->GetRenderTargetItem().ShaderResourceTexture,  ShadingEnergyConservationData::GGXSpecValues, 2);
-				ShadingEnergyConservationData::LockCopyTexture3D(GraphBuilder.RHICmdList, (FTexture3DRHIRef&)View.ViewState->ShadingEnergyConservationData.GGXGlassEnergyTexture->GetRenderTargetItem().ShaderResourceTexture, ShadingEnergyConservationData::GGXGlassValues, 2);
-				ShadingEnergyConservationData::LockCopyTexture2D(GraphBuilder.RHICmdList, (FTexture2DRHIRef&)View.ViewState->ShadingEnergyConservationData.ClothEnergyTexture->GetRenderTargetItem().ShaderResourceTexture,    Strata::IsStrataEnabled() ? ShadingEnergyConservationData::StrataClothSpecValues : ShadingEnergyConservationData::ClothSpecValues, 2);
-				ShadingEnergyConservationData::LockCopyTexture2D(GraphBuilder.RHICmdList, (FTexture2DRHIRef&)View.ViewState->ShadingEnergyConservationData.DiffuseEnergyTexture->GetRenderTargetItem().ShaderResourceTexture,  ShadingEnergyConservationData::DiffuseValues, 1);
+				ShadingEnergyConservationData::LockCopyTexture2D(GraphBuilder.RHICmdList, View.ViewState->ShadingEnergyConservationData.GGXSpecEnergyTexture->GetRHI(),  ShadingEnergyConservationData::GGXSpecValues, 2);
+				ShadingEnergyConservationData::LockCopyTexture3D(GraphBuilder.RHICmdList, View.ViewState->ShadingEnergyConservationData.GGXGlassEnergyTexture->GetRHI(), ShadingEnergyConservationData::GGXGlassValues, 2);
+				ShadingEnergyConservationData::LockCopyTexture2D(GraphBuilder.RHICmdList, View.ViewState->ShadingEnergyConservationData.ClothEnergyTexture->GetRHI(),    Strata::IsStrataEnabled() ? ShadingEnergyConservationData::StrataClothSpecValues : ShadingEnergyConservationData::ClothSpecValues, 2);
+				ShadingEnergyConservationData::LockCopyTexture2D(GraphBuilder.RHICmdList, View.ViewState->ShadingEnergyConservationData.DiffuseEnergyTexture->GetRHI(),  ShadingEnergyConservationData::DiffuseValues, 1);
 
 				GGXSpecEnergyTexture  = GraphBuilder.RegisterExternalTexture(View.ViewState->ShadingEnergyConservationData.GGXSpecEnergyTexture);
 				GGXGlassEnergyTexture = GraphBuilder.RegisterExternalTexture(View.ViewState->ShadingEnergyConservationData.GGXGlassEnergyTexture);
diff --git a/Engine/Source/Runtime/Renderer/Private/ShadowRendering.cpp b/Engine/Source/Runtime/Renderer/Private/ShadowRendering.cpp
--- a/Engine/Source/Runtime/Renderer/Private/ShadowRendering.cpp
+++ b/Engine/Source/Runtime/Renderer/Private/ShadowRendering.cpp
@@ -2476,8 +2476,8 @@ void SetupTranslucentSelfShadowUniformParameters(const FProjectedShadowInfo* Sha
 		// Incorporate the diffuse scale of 1 / PI into the light color
 		OutParameters.DirectionalLightColor = FVector4f(FVector3f(LightProxy->GetColor() * FadeAlpha / PI), FadeAlpha);
 
-		OutParameters.Transmission0 = ShadowInfo->RenderTargets.ColorTargets[0]->GetRenderTargetItem().ShaderResourceTexture.GetReference();
-		OutParameters.Transmission1 = ShadowInfo->RenderTargets.ColorTargets[1]->GetRenderTargetItem().ShaderResourceTexture.GetReference();
+		OutParameters.Transmission0 = ShadowInfo->RenderTargets.ColorTargets[0]->GetRHI();
+		OutParameters.Transmission1 = ShadowInfo->RenderTargets.ColorTargets[1]->GetRHI();
 		OutParameters.Transmission0Sampler = TStaticSamplerState<SF_Bilinear, AM_Clamp, AM_Clamp, AM_Clamp>::GetRHI();
 		OutParameters.Transmission1Sampler = TStaticSamplerState<SF_Bilinear, AM_Clamp, AM_Clamp, AM_Clamp>::GetRHI();
 	}
diff --git a/Engine/Source/Runtime/Renderer/Private/ShadowRendering.h b/Engine/Source/Runtime/Renderer/Private/ShadowRendering.h
--- a/Engine/Source/Runtime/Renderer/Private/ShadowRendering.h
+++ b/Engine/Source/Runtime/Renderer/Private/ShadowRendering.h
@@ -1054,11 +1054,11 @@ public:
 		// Translucency shadow projection has no depth target
 		if (ShadowInfo->RenderTargets.DepthTarget)
 		{
-			ShadowDepthTextureValue = ShadowInfo->RenderTargets.DepthTarget->GetRenderTargetItem().ShaderResourceTexture.GetReference();
+			ShadowDepthTextureValue = ShadowInfo->RenderTargets.DepthTarget->GetRHI();
 		}
 		else
 		{
-			ShadowDepthTextureValue = GSystemTextures.BlackDummy->GetRenderTargetItem().ShaderResourceTexture.GetReference();
+			ShadowDepthTextureValue = GSystemTextures.BlackDummy->GetRHI();
 		}
 			
 		FRHISamplerState* DepthSamplerState = TStaticSamplerState<SF_Point,AM_Clamp,AM_Clamp,AM_Clamp>::GetRHI();
@@ -1367,12 +1367,14 @@ public:
 	{
 		const FVector PreViewTranslation = View.ViewMatrices.GetPreViewTranslation();
 
-		FRHITexture* ShadowDepthTextureValue = ShadowInfo
-			? ShadowInfo->RenderTargets.DepthTarget->GetRenderTargetItem().ShaderResourceTexture->GetTextureCube()
-			: GBlackTextureDepthCube->TextureRHI.GetReference();
-		if (!ShadowDepthTextureValue)
+		FRHITexture* ShadowDepthTextureValue = GBlackTextureDepthCube->TextureRHI;
+		
+		if (ShadowInfo)
 		{
-			ShadowDepthTextureValue = GBlackTextureDepthCube->TextureRHI.GetReference();
+			if (FRHITexture* Texture = ShadowInfo->RenderTargets.DepthTarget->GetRHI())
+			{
+				ShadowDepthTextureValue = Texture;
+			}
 		}
 
 		SetTextureParameter(
diff --git a/Engine/Source/Runtime/Renderer/Private/SingleLayerWaterRendering.cpp b/Engine/Source/Runtime/Renderer/Private/SingleLayerWaterRendering.cpp
--- a/Engine/Source/Runtime/Renderer/Private/SingleLayerWaterRendering.cpp
+++ b/Engine/Source/Runtime/Renderer/Private/SingleLayerWaterRendering.cpp
@@ -415,7 +415,7 @@ void FDeferredShadingSceneRenderer::RenderSingleLayerWaterReflections(
 		{
 			Parameters.ScreenSpaceReflectionsTexture = ReflectionsColor ? ReflectionsColor : BlackDummyTexture;
 			Parameters.ScreenSpaceReflectionsSampler = TStaticSamplerState<SF_Point>::GetRHI();
-			Parameters.PreIntegratedGF = GSystemTextures.PreintegratedGF->GetRenderTargetItem().ShaderResourceTexture;
+			Parameters.PreIntegratedGF = GSystemTextures.PreintegratedGF->GetRHI();
 			Parameters.PreIntegratedGFSampler = TStaticSamplerState<SF_Bilinear, AM_Clamp, AM_Clamp, AM_Clamp>::GetRHI();
 			Parameters.SceneNoWaterDepthTexture = SceneWithoutWaterTextures.DepthTexture ? SceneWithoutWaterTextures.DepthTexture : BlackDummyTexture;
 			Parameters.SceneNoWaterDepthSampler = TStaticSamplerState<SF_Point>::GetRHI();
diff --git a/Engine/Source/Runtime/Renderer/Private/SkyAtmosphereRendering.cpp b/Engine/Source/Runtime/Renderer/Private/SkyAtmosphereRendering.cpp
--- a/Engine/Source/Runtime/Renderer/Private/SkyAtmosphereRendering.cpp
+++ b/Engine/Source/Runtime/Renderer/Private/SkyAtmosphereRendering.cpp
@@ -382,7 +382,7 @@ void SetupSkyAtmosphereViewSharedUniformShaderParameters(const FViewInfo& View,
 	FRHITexture* SkyAtmosphereCameraAerialPerspectiveVolume = nullptr;
 	if (View.SkyAtmosphereCameraAerialPerspectiveVolume)
 	{
-		SkyAtmosphereCameraAerialPerspectiveVolume = View.SkyAtmosphereCameraAerialPerspectiveVolume->GetRenderTargetItem().ShaderResourceTexture;
+		SkyAtmosphereCameraAerialPerspectiveVolume = View.SkyAtmosphereCameraAerialPerspectiveVolume->GetRHI();
 	}
 
 	OutParameters.CameraAerialPerspectiveVolumeSizeAndInvSize = GetSizeAndInvSize(CameraAerialPerspectiveVolumeScreenResolution, CameraAerialPerspectiveVolumeScreenResolution);
@@ -2013,7 +2013,7 @@ void FSceneRenderer::RenderSkyAtmosphereEditorNotifications(FRDGBuilder& GraphBu
 			RenderSkyAtmosphereEditorHudPS::FParameters* PassParameters = GraphBuilder.AllocParameters<RenderSkyAtmosphereEditorHudPS::FParameters>();
 			PassParameters->RenderTargets[0] = FRenderTargetBinding(SceneColorTexture, ERenderTargetLoadAction::ELoad);
 			PassParameters->ViewUniformBuffer = View.ViewUniformBuffer;
-			PassParameters->MiniFontTexture = GSystemTextures.AsciiTexture->GetRenderTargetItem().ShaderResourceTexture;
+			PassParameters->MiniFontTexture = GSystemTextures.AsciiTexture->GetRHI();
 
 			FPixelShaderUtils::AddFullscreenPass<RenderSkyAtmosphereEditorHudPS>(GraphBuilder, View.ShaderMap, RDG_EVENT_NAME("SkyAtmosphereEditor"), PixelShader, PassParameters, View.ViewRect);
 		}
diff --git a/Engine/Source/Runtime/Renderer/Private/SystemTextures.cpp b/Engine/Source/Runtime/Renderer/Private/SystemTextures.cpp
--- a/Engine/Source/Runtime/Renderer/Private/SystemTextures.cpp
+++ b/Engine/Source/Runtime/Renderer/Private/SystemTextures.cpp
@@ -310,10 +310,10 @@ void FSystemTextures::InitializeFeatureLevelDependentTextures(FRHICommandListImm
 		// Init with dummy textures. The texture will be initialize with real values if needed
 		const uint8 BlackBytes[4] = { 0, 0, 0, 0 };
 		FUpdateTextureRegion3D Region(0, 0, 0, 0, 0, 0, Desc.Extent.X, Desc.Extent.Y, Desc.Depth);
-		RHICmdList.UpdateTexture3D((FTexture3DRHIRef&)HairLUT0->GetRenderTargetItem().ShaderResourceTexture, 0, Region, Desc.Extent.X * sizeof(BlackBytes), Desc.Extent.X * Desc.Extent.Y * sizeof(BlackBytes), BlackBytes);
+		RHICmdList.UpdateTexture3D(HairLUT0->GetRHI(), 0, Region, Desc.Extent.X * sizeof(BlackBytes), Desc.Extent.X * Desc.Extent.Y * sizeof(BlackBytes), BlackBytes);
 
 		// UpdateTexture3D before and after state is currently undefined
-		RHICmdList.Transition(FRHITransitionInfo(HairLUT0->GetRenderTargetItem().ShaderResourceTexture, ERHIAccess::Unknown, ERHIAccess::SRVMask));
+		RHICmdList.Transition(FRHITransitionInfo(HairLUT0->GetRHI(), ERHIAccess::Unknown, ERHIAccess::SRVMask));
 		HairLUT1 = HairLUT0;
 		HairLUT2 = HairLUT0;
 	}
@@ -358,7 +358,7 @@ void FSystemTextures::InitializeFeatureLevelDependentTextures(FRHICommandListImm
 
 		GRenderTargetPool.FindFreeElement(RHICmdList, Desc, AsciiTexture, TEXT("AsciiTexture"));
 		uint32 DestStride;
-		uint8* DestBuffer = (uint8*)RHICmdList.LockTexture2D((FTexture2DRHIRef&)AsciiTexture->GetRenderTargetItem().ShaderResourceTexture, 0, RLM_WriteOnly, DestStride, false);
+		uint8* DestBuffer = (uint8*)RHICmdList.LockTexture2D(AsciiTexture->GetRHI(), 0, RLM_WriteOnly, DestStride, false);
 
 		for (int32 c = 0; c < CharacterCount; c++)
 		{
@@ -375,7 +375,7 @@ void FSystemTextures::InitializeFeatureLevelDependentTextures(FRHICommandListImm
 				Dest[0] = bFilled ? 0xFF : 0x0;
 			}
 		}
-		RHICmdList.UnlockTexture2D((FTexture2DRHIRef&)AsciiTexture->GetRenderTargetItem().ShaderResourceTexture, 0, false);
+		RHICmdList.UnlockTexture2D(AsciiTexture->GetRHI(), 0, false);
 	}
 
 	// The PreintegratedGF maybe used on forward shading including mobile platform, initialize it anyway.
diff --git a/Engine/Source/Runtime/Renderer/Private/TranslucentRendering.cpp b/Engine/Source/Runtime/Renderer/Private/TranslucentRendering.cpp
--- a/Engine/Source/Runtime/Renderer/Private/TranslucentRendering.cpp
+++ b/Engine/Source/Runtime/Renderer/Private/TranslucentRendering.cpp
@@ -995,15 +995,15 @@ TRDGUniformBufferRef<FTranslucentBasePassUniformParameters> CreateTranslucentBas
 			{
 				TRefCountPtr<IPooledRenderTarget> VolumetricReconstructRTDepth = View.ViewState->VolumetricCloudRenderTarget.GetDstVolumetricReconstructRTDepth();
 
-				BasePassParameters.VolumetricCloudColor = VolumetricReconstructRT->GetRenderTargetItem().ShaderResourceTexture;
-				BasePassParameters.VolumetricCloudDepth = VolumetricReconstructRTDepth->GetRenderTargetItem().ShaderResourceTexture;
+				BasePassParameters.VolumetricCloudColor = VolumetricReconstructRT->GetRHI();
+				BasePassParameters.VolumetricCloudDepth = VolumetricReconstructRTDepth->GetRHI();
 				BasePassParameters.ApplyVolumetricCloudOnTransparent = 1.0f;
 			}
 		}
 		if (BasePassParameters.VolumetricCloudColor == nullptr)
 		{
-			BasePassParameters.VolumetricCloudColor = GSystemTextures.BlackAlphaOneDummy->GetRenderTargetItem().ShaderResourceTexture;
-			BasePassParameters.VolumetricCloudDepth = GSystemTextures.BlackDummy->GetRenderTargetItem().ShaderResourceTexture;
+			BasePassParameters.VolumetricCloudColor = GSystemTextures.BlackAlphaOneDummy->GetRHI();
+			BasePassParameters.VolumetricCloudDepth = GSystemTextures.BlackDummy->GetRHI();
 		}
 
 		FIntPoint ViewportOffset = View.ViewRect.Min;
@@ -1070,7 +1070,7 @@ TRDGUniformBufferRef<FTranslucentBasePassUniformParameters> CreateTranslucentBas
 	}
 
 	BasePassParameters.EyeAdaptationTexture = GetEyeAdaptationTexture(GraphBuilder, View);
-	BasePassParameters.PreIntegratedGFTexture = GSystemTextures.PreintegratedGF->GetRenderTargetItem().ShaderResourceTexture;
+	BasePassParameters.PreIntegratedGFTexture = GSystemTextures.PreintegratedGF->GetRHI();
 	BasePassParameters.PreIntegratedGFSampler = TStaticSamplerState<SF_Bilinear, AM_Clamp, AM_Clamp, AM_Clamp>::GetRHI();
 
 	OIT::SetOITParameters(GraphBuilder, View, BasePassParameters.OIT, OITData);
diff --git a/Engine/Source/Runtime/Renderer/Private/VT/VirtualTexturePhysicalSpace.cpp b/Engine/Source/Runtime/Renderer/Private/VT/VirtualTexturePhysicalSpace.cpp
--- a/Engine/Source/Runtime/Renderer/Private/VT/VirtualTexturePhysicalSpace.cpp
+++ b/Engine/Source/Runtime/Renderer/Private/VT/VirtualTexturePhysicalSpace.cpp
@@ -139,7 +139,7 @@ void FVirtualTexturePhysicalSpace::InitRHI()
 		}
 
 		GRenderTargetPool.FindFreeElement(RHICmdList, Desc, PooledRenderTarget[Layer], TEXT("VirtualPhysicalTexture"));
-		FRHITexture* TextureRHI = PooledRenderTarget[Layer]->GetRenderTargetItem().ShaderResourceTexture;
+		FRHITexture* TextureRHI = PooledRenderTarget[Layer]->GetRHI();
 
 		// Create sRGB and non-sRGB shader resource views into the physical texture
 		FRHITextureSRVCreateInfo SRVCreateInfo;
@@ -152,7 +152,7 @@ void FVirtualTexturePhysicalSpace::InitRHI()
 
 		if (bCreateAliasedUAV)
 		{
-			TextureUAV[Layer] = PooledRenderTarget[Layer]->GetRenderTargetItem().UAV;
+			TextureUAV[Layer] = PooledRenderTarget[Layer]->GetUAV();
 		}
 	}
 }
diff --git a/Engine/Source/Runtime/Renderer/Private/VT/VirtualTexturePhysicalSpace.h b/Engine/Source/Runtime/Renderer/Private/VT/VirtualTexturePhysicalSpace.h
--- a/Engine/Source/Runtime/Renderer/Private/VT/VirtualTexturePhysicalSpace.h
+++ b/Engine/Source/Runtime/Renderer/Private/VT/VirtualTexturePhysicalSpace.h
@@ -89,7 +89,7 @@ public:
 	FRHITexture* GetPhysicalTexture(int32 Layer) const
 	{
 		check(PooledRenderTarget[Layer].IsValid());
-		return PooledRenderTarget[Layer]->GetRenderTargetItem().ShaderResourceTexture;
+		return PooledRenderTarget[Layer]->GetRHI();
 	}
 
 	FRHIShaderResourceView* GetPhysicalTextureSRV(int32 Layer, bool bSRGB) const
diff --git a/Engine/Source/Runtime/Renderer/Private/VT/VirtualTextureTest.cpp b/Engine/Source/Runtime/Renderer/Private/VT/VirtualTextureTest.cpp
--- a/Engine/Source/Runtime/Renderer/Private/VT/VirtualTextureTest.cpp
+++ b/Engine/Source/Runtime/Renderer/Private/VT/VirtualTextureTest.cpp
@@ -185,9 +185,9 @@ void FVirtualTextureTest::ProducePageData( FRHICommandList& RHICmdList, ERHIFeat
 
 	auto ShaderMap = GetGlobalShaderMap( FeatureLevel );
 
-	FSceneRenderTargetItem& RenderTarget = PhysicalTexture->GetRenderTargetItem();
+	FRHITexture* RenderTarget = PhysicalTexture->GetRHI();
 
-	FRHIRenderPassInfo RPInfo(RenderTarget.TargetableTexture, ERenderTargetActions::Load_Store);
+	FRHIRenderPassInfo RPInfo(RenderTarget, ERenderTargetActions::Load_Store);
 	RHICmdList.BeginRenderPass(RPInfo, TEXT("ProducePageData"));
 	{
 		RHICmdList.SetViewport(DstRect.Min.X, DstRect.Min.Y, 0.0f, DstRect.Max.X, DstRect.Max.Y, 1.0f);
@@ -220,7 +220,7 @@ void FVirtualTextureTest::ProducePageData( FRHICommandList& RHICmdList, ERHIFeat
 			EDRF_UseTriangleOptimization);
 	}
 	RHICmdList.EndRenderPass();
-	RHICmdList.CopyToResolveTarget(RenderTarget.TargetableTexture, RenderTarget.ShaderResourceTexture, FResolveParams());
+	RHICmdList.Transition(FRHITransitionInfo(RenderTarget, ERHIAccess::RTV, ERHIAccess::SRVMask));
 
 	GVisualizeTexture.SetCheckPoint( RHICmdList, PhysicalTexture );
 }
diff --git a/Engine/Source/Runtime/Renderer/Private/VT/VirtualTextureTest.h b/Engine/Source/Runtime/Renderer/Private/VT/VirtualTextureTest.h
--- a/Engine/Source/Runtime/Renderer/Private/VT/VirtualTextureTest.h
+++ b/Engine/Source/Runtime/Renderer/Private/VT/VirtualTextureTest.h
@@ -35,7 +35,7 @@ public:
 	virtual void	InitDynamicRHI() override;
 	virtual void	ReleaseDynamicRHI() override;
 	
-	FRHITexture*	GetPhysicalTexture() const		{ return PhysicalTexture->GetRenderTargetItem().ShaderResourceTexture; }
+	FRHITexture*	GetPhysicalTexture() const		{ return PhysicalTexture->GetRHI(); }
 	FRHITexture*	GetPageTableTexture() const;
 	
 	FTexturePagePool*		Pool;
diff --git a/Engine/Source/Runtime/SlateRHIRenderer/Private/SlateRHIRenderer.cpp b/Engine/Source/Runtime/SlateRHIRenderer/Private/SlateRHIRenderer.cpp
--- a/Engine/Source/Runtime/SlateRHIRenderer/Private/SlateRHIRenderer.cpp
+++ b/Engine/Source/Runtime/SlateRHIRenderer/Private/SlateRHIRenderer.cpp
@@ -795,8 +795,8 @@ void FSlateRHIRenderer::DrawWindow_RenderThread(FRHICommandListImmediate& RHICmd
 			const uint32 ViewportHeight = (ViewportRT) ? ViewportRT->GetSizeY() : ViewportInfo.Height;
 
 			// Check to see that targets are up-to-date
-			if (bCompositeUI && (!ViewportInfo.UITargetRT || ViewportInfo.UITargetRT->GetRenderTargetItem().TargetableTexture->GetTexture2D()->GetSizeX() != ViewportWidth || ViewportInfo.UITargetRT->GetRenderTargetItem().TargetableTexture->GetTexture2D()->GetSizeY() != ViewportHeight
-				|| !ViewportInfo.HDRSourceRT || ViewportInfo.HDRSourceRT->GetRenderTargetItem().TargetableTexture->GetFormat() != BackBuffer->GetFormat()))
+			if (bCompositeUI && (!ViewportInfo.UITargetRT || ViewportInfo.UITargetRT->GetRHI()->GetSizeX() != ViewportWidth || ViewportInfo.UITargetRT->GetRHI()->GetSizeY() != ViewportHeight
+				|| !ViewportInfo.HDRSourceRT || ViewportInfo.HDRSourceRT->GetRHI()->GetFormat() != BackBuffer->GetFormat()))
 			{
 				// Composition buffers
 				{
@@ -843,16 +843,16 @@ void FSlateRHIRenderer::DrawWindow_RenderThread(FRHICommandListImmediate& RHICmd
 
 				FRHITransitionInfo Transitions[] = {
 					FRHITransitionInfo(FinalBuffer, ERHIAccess::Unknown, ERHIAccess::ResolveSrc),
-					FRHITransitionInfo(ViewportInfo.HDRSourceRT->GetRenderTargetItem().TargetableTexture, ERHIAccess::Unknown, ERHIAccess::ResolveDst)
+					FRHITransitionInfo(ViewportInfo.HDRSourceRT->GetRHI(), ERHIAccess::Unknown, ERHIAccess::ResolveDst)
 				};
 				RHICmdList.Transition(MakeArrayView(Transitions, UE_ARRAY_COUNT(Transitions)));
 
 				// Grab HDR backbuffer
 				FResolveParams ResolveParams;
-				RHICmdList.CopyToResolveTarget(FinalBuffer, ViewportInfo.HDRSourceRT->GetRenderTargetItem().TargetableTexture, ResolveParams);
+				RHICmdList.CopyToResolveTarget(FinalBuffer, ViewportInfo.HDRSourceRT->GetRHI(), ResolveParams);
 
 				// UI backbuffer is temp target
-				BackBuffer = ViewportInfo.UITargetRT->GetRenderTargetItem().TargetableTexture->GetTexture2D();
+				BackBuffer = ViewportInfo.UITargetRT->GetRHI();
 			}
 
 #if WITH_EDITOR
@@ -875,7 +875,7 @@ void FSlateRHIRenderer::DrawWindow_RenderThread(FRHICommandListImmediate& RHICmd
 				// FIXME: is this necessary?
 				RHICmdList.Transition(FRHITransitionInfo(FinalBuffer, ERHIAccess::Unknown, ERHIAccess::SRVMask));
 
-				BackBuffer = HDRRenderRT->GetRenderTargetItem().TargetableTexture->GetTexture2D();
+				BackBuffer = HDRRenderRT->GetRHI();
 			}
 #endif
 
@@ -1015,8 +1015,7 @@ void FSlateRHIRenderer::DrawWindow_RenderThread(FRHICommandListImmediate& RHICmd
 
 				// Composition pass
 				{
-					FResolveParams ResolveParams;
-					RHICmdList.CopyToResolveTarget(ViewportInfo.UITargetRT->GetRenderTargetItem().TargetableTexture, ViewportInfo.UITargetRT->GetRenderTargetItem().TargetableTexture, ResolveParams);
+					RHICmdList.Transition(FRHITransitionInfo(ViewportInfo.UITargetRT->GetRHI(), ERHIAccess::Unknown, ERHIAccess::SRVMask));
 
 					if (RHISupportsRenderTargetWriteMask(GMaxRHIShaderPlatform))
 					{
@@ -1024,11 +1023,10 @@ void FSlateRHIRenderer::DrawWindow_RenderThread(FRHICommandListImmediate& RHICmd
 						FRenderTargetWriteMask::Decode(RHICmdList, ShaderMap, RenderTargets, ViewportInfo.UITargetRTMask, TexCreate_None, TEXT("UIRTWriteMask"));
 					}
 
-					FRHITransitionInfo Transitions[] = {
+					RHICmdList.Transition({
 						FRHITransitionInfo(FinalBuffer, ERHIAccess::Unknown, ERHIAccess::RTV),
-						FRHITransitionInfo(ViewportInfo.HDRSourceRT->GetRenderTargetItem().TargetableTexture, ERHIAccess::Unknown, ERHIAccess::SRVGraphics)
-					};
-					RHICmdList.Transition(MakeArrayView(Transitions, UE_ARRAY_COUNT(Transitions)));
+						FRHITransitionInfo(ViewportInfo.HDRSourceRT->GetRHI(), ERHIAccess::Unknown, ERHIAccess::SRVGraphics)
+					});
 					FRHIRenderPassInfo RPInfo(FinalBuffer, ERenderTargetActions::Load_Store);
 					RHICmdList.BeginRenderPass(RPInfo, TEXT("SlateComposite"));
 					{
@@ -1040,7 +1038,7 @@ void FSlateRHIRenderer::DrawWindow_RenderThread(FRHICommandListImmediate& RHICmd
 
 						TShaderMapRef<FScreenVS> VertexShader(ShaderMap);
 
-						FRHITexture* UITargetRTMaskTexture = RHISupportsRenderTargetWriteMask(GMaxRHIShaderPlatform) ? ViewportInfo.UITargetRTMask->GetRenderTargetItem().TargetableTexture : nullptr;
+						FRHITexture* UITargetRTMaskTexture = RHISupportsRenderTargetWriteMask(GMaxRHIShaderPlatform) ? ViewportInfo.UITargetRTMask->GetRHI() : nullptr;
 						if (HDROutputDevice == 5 || HDROutputDevice == 6)
 						{
 							// ScRGB encoding
@@ -1053,7 +1051,7 @@ void FSlateRHIRenderer::DrawWindow_RenderThread(FRHICommandListImmediate& RHICmd
 
 							SetGraphicsPipelineState(RHICmdList, GraphicsPSOInit, 0);
 
-							PixelShader->SetParameters(RHICmdList, ViewportInfo.UITargetRT->GetRenderTargetItem().TargetableTexture, UITargetRTMaskTexture, ViewportInfo.HDRSourceRT->GetRenderTargetItem().TargetableTexture, ViewportInfo.ColorSpaceLUTSRV);
+							PixelShader->SetParameters(RHICmdList, ViewportInfo.UITargetRT->GetRHI(), UITargetRTMaskTexture, ViewportInfo.HDRSourceRT->GetRHI(), ViewportInfo.ColorSpaceLUTSRV);
 						}
 						else
 						{
@@ -1067,7 +1065,7 @@ void FSlateRHIRenderer::DrawWindow_RenderThread(FRHICommandListImmediate& RHICmd
 
 							SetGraphicsPipelineState(RHICmdList, GraphicsPSOInit, 0);
 
-							PixelShader->SetParameters(RHICmdList, ViewportInfo.UITargetRT->GetRenderTargetItem().TargetableTexture, UITargetRTMaskTexture, ViewportInfo.HDRSourceRT->GetRenderTargetItem().TargetableTexture, ViewportInfo.ColorSpaceLUTSRV);
+							PixelShader->SetParameters(RHICmdList, ViewportInfo.UITargetRT->GetRHI(), UITargetRTMaskTexture, ViewportInfo.HDRSourceRT->GetRHI(), ViewportInfo.ColorSpaceLUTSRV);
 						}
 
 						RendererModule.DrawRectangle(
@@ -1116,7 +1114,7 @@ void FSlateRHIRenderer::DrawWindow_RenderThread(FRHICommandListImmediate& RHICmd
 
 				SetGraphicsPipelineState(RHICmdList, GraphicsPSOInit, 0);
 
-				PixelShader->SetParameters(RHICmdList, HDRRenderRT->GetRenderTargetItem().TargetableTexture );
+				PixelShader->SetParameters(RHICmdList, HDRRenderRT->GetRHI() );
 
 				static const FName RendererModuleName("Renderer");
 				IRendererModule& RendererModule = FModuleManager::GetModuleChecked<IRendererModule>(RendererModuleName);
