GPULM: Fix occasional SWidget related crash when exiting the editor without ending the lighting build first

diff --git a/Engine/Plugins/Experimental/GPULightmass/Source/GPULightmass/Private/GPULightmass.cpp b/Engine/Plugins/Experimental/GPULightmass/Source/GPULightmass/Private/GPULightmass.cpp
--- a/Engine/Plugins/Experimental/GPULightmass/Source/GPULightmass/Private/GPULightmass.cpp
+++ b/Engine/Plugins/Experimental/GPULightmass/Source/GPULightmass/Private/GPULightmass.cpp
@@ -69,12 +69,21 @@ void FGPULightmass::GameThreadDestroy()
 
 	RemoveGameThreadEventHooks();
 
-	if (!IsEngineExitRequested() && LightBuildNotification.IsValid())
+	if (LightBuildNotification.IsValid())
 	{
-		FText CompletedText = LOCTEXT("LightBuildDoneMessage", "Lighting build completed");
-		LightBuildNotification->SetText(CompletedText);
-		LightBuildNotification->SetCompletionState(SNotificationItem::CS_Success);
-		LightBuildNotification->ExpireAndFadeout();
+		if (!IsEngineExitRequested())
+		{
+			// Shows a notification that fades out slowly
+			FText CompletedText = LOCTEXT("LightBuildDoneMessage", "Lighting build completed");
+			LightBuildNotification->SetText(CompletedText);
+			LightBuildNotification->SetCompletionState(SNotificationItem::CS_Success);
+			LightBuildNotification->ExpireAndFadeout();
+		}
+		else
+		{
+			// Immediately destroys the notification widget to avoid crash
+			LightBuildNotification.Reset();
+		}
 	}
 
 	Scene.RemoveAllComponents();
diff --git a/Engine/Plugins/Experimental/GPULightmass/Source/GPULightmass/Private/GPULightmass.h b/Engine/Plugins/Experimental/GPULightmass/Source/GPULightmass/Private/GPULightmass.h
--- a/Engine/Plugins/Experimental/GPULightmass/Source/GPULightmass/Private/GPULightmass.h
+++ b/Engine/Plugins/Experimental/GPULightmass/Source/GPULightmass/Private/GPULightmass.h
@@ -34,6 +34,7 @@ public:
 	void StartRecordingVisibleTiles();
 	void EndRecordingVisibleTiles();
 
+	// Game thread objects. Must be destroyed in GameThreadDestroy()
 	TSharedPtr<SNotificationItem> LightBuildNotification;
 	int32 LightBuildPercentage;
 	
