GPULM: Fix denoiser discoloration issue by removing the no longer needed inputScale

diff --git a/Engine/Plugins/Experimental/GPULightmass/Source/GPULightmass/Private/LightmapDenoising.cpp b/Engine/Plugins/Experimental/GPULightmass/Source/GPULightmass/Private/LightmapDenoising.cpp
--- a/Engine/Plugins/Experimental/GPULightmass/Source/GPULightmass/Private/LightmapDenoising.cpp
+++ b/Engine/Plugins/Experimental/GPULightmass/Source/GPULightmass/Private/LightmapDenoising.cpp
@@ -20,7 +20,6 @@ FDenoiserFilterSet::FDenoiserFilterSet(FDenoiserContext& Context, FIntPoint NewS
 	filter.setImage("color", InputBuffer.GetData(), oidn::Format::Float3, Size.X, Size.Y);
 	filter.setImage("output", OutputBuffer.GetData(), oidn::Format::Float3, Size.X, Size.Y);
 	filter.set("directional", bSHDenoiser);
-	filter.set("inputScale", bSHDenoiser ? 0.5f : 1.0f);
 	filter.commit();
 #endif
 
