GPU Lightmass Solver Improvements

Irradiance Caching
* Instead of setting a per-cell cap of samples, now we use the first X path samples of the receiver (lightmap texels) to adaptively distribute the irradiance cache samples. This will distribute them where BRDF sampling thinks the most important (more around corners and fewer in open area, for example). This is also the distribution that the old school irradiance caching's distance-based heuristics approximates for.
* New cache cells are only spawned during the first X passes. This avoids the case where the cache is constantly thrashed when it is not big enough to cover the scene.
* The image is reset after the first X passes are done. Like path guiding, these passes are uncached and thus contain high variance. Mixing them with the passes later will actually increase overall variance.

First Bounce Ray Guiding
* Now the FBRG passes run after IC passes are done. Way less input noise into the guiding PDF and better at locating importance sources.

New Parameter Tweaking Strategy for IC and FBRG
* IC: set quality to 1/16-1/8 of GI samples. Increase when seeing bright splotches around corners, or color discrepancies over large area.
* FBRE: set trial samples to 1/4 of GI samples. Increase when fireflies or small square-shaped artifacts are introduced by this technique.

General Cleanups
* IC and FBRG passes now add to the total number of passes to render instead of being subtracted from the number of GI Samples to reduce user confusion.
* IC buffer size reduction to save some GPU memory.

diff --git a/Engine/Plugins/Experimental/GPULightmass/Shaders/Private/FirstBounceRayGuidingCDFBuild.usf b/Engine/Plugins/Experimental/GPULightmass/Shaders/Private/FirstBounceRayGuidingCDFBuild.usf
--- a/Engine/Plugins/Experimental/GPULightmass/Shaders/Private/FirstBounceRayGuidingCDFBuild.usf
+++ b/Engine/Plugins/Experimental/GPULightmass/Shaders/Private/FirstBounceRayGuidingCDFBuild.usf
@@ -4,7 +4,7 @@
 
 #include "BatchedTiles.ush"
 
-int NumRayGuidingTrialSamples;
+int RayGuidingEndPassIndex;
 
 RWTexture2D<uint> RayGuidingLuminance;
 RWTexture2D<float> RayGuidingCDFX;
@@ -22,7 +22,7 @@ void FirstBounceRayGuidingCDFBuildCS(
 	int ClusterLinearIndexInTile = GroupID.x % NUM_CLUSTERS_PER_TILE;
 	int2 ClusterPosition = int2(ClusterLinearIndexInTile % NUM_CLUSTERS_PER_DIM, ClusterLinearIndexInTile / NUM_CLUSTERS_PER_DIM);
 	
-	if (BatchedTiles[TileIndex].RenderPassIndex != NumRayGuidingTrialSamples - 1) return;
+	if (BatchedTiles[TileIndex].RenderPassIndex != RayGuidingEndPassIndex) return;
 	
 	if (ThreadId.x < DIRECTIONAL_BINS_ONE_DIM)
 	{
@@ -34,7 +34,7 @@ void FirstBounceRayGuidingCDFBuildCS(
 			for(int dy = -FilterKernelSize; dy <= FilterKernelSize; dy++)
 			{
 				int2 FinalPosition = ClusterPosition * DIRECTIONAL_BINS_ONE_DIM + clamp(ThreadId.xy + int2(dx, dy), int2(0, 0), int2(DIRECTIONAL_BINS_ONE_DIM - 1, DIRECTIONAL_BINS_ONE_DIM - 1));
-				Value += asfloat(RayGuidingLuminance[BatchedTiles[TileIndex].WorkingSetPosition / GPreviewLightmapPhysicalTileSize * CDF_TILE_SIZE + FinalPosition]);
+				Value += asfloat(RayGuidingLuminance[BatchedTiles[TileIndex].WorkingSetPosition / GPreviewLightmapPhysicalTileSize * CDF_TILE_SIZE + FinalPosition]) * ((dx != 0 || dy != 0) ? 0.125 : 1);
 			}
 		}
 
diff --git a/Engine/Plugins/Experimental/GPULightmass/Shaders/Private/IrradianceCacheVisualization.usf b/Engine/Plugins/Experimental/GPULightmass/Shaders/Private/IrradianceCacheVisualization.usf
--- a/Engine/Plugins/Experimental/GPULightmass/Shaders/Private/IrradianceCacheVisualization.usf
+++ b/Engine/Plugins/Experimental/GPULightmass/Shaders/Private/IrradianceCacheVisualization.usf
@@ -5,13 +5,6 @@
 =============================================================================*/
 
 // Needs to be defined before the uniform buffer struct header is included, very unfortunately
-struct FIrradianceCacheRecord
-{
-	float4 WorldPosition;
-	float4 WorldNormal;
-	uint4 Irradiance;
-};
-
 #include "/Engine/Private/Common.ush"
 #include "/Engine/Private/DeferredShadingCommon.ush"
 #include "IrradianceCachingCommon.ush"
@@ -32,61 +25,29 @@ void VisualizeIrradianceCachePS(
 			
 	uint NearestRecordIndex = 0;
 	float3 RecordIrradiance;
-	float3 RecordWorldNormal;
-			
-	uint CoarseLevelSpacing = IrradianceCachingParameters.Spacing;
-	uint DetailLevelSpacing = IrradianceCachingParameters.Spacing;
-		
-	uint3 Key = EncodeVoxelKey(floor(WorldPosition / DetailLevelSpacing));		
+	
+	uint3 Key = EncodeVoxelKey(EncodePositionKey(WorldPosition, ShadingNormal, IrradianceCachingParameters.Spacing));		
 	uint NormalDirectionBits = EncodeNormalBits(ShadingNormal);				
 	uint HashKey = 0;
-	HashKey = MurmurAdd( HashKey, NormalDirectionBits );
 	HashKey = MurmurAdd( HashKey, Key.x );
 	HashKey = MurmurAdd( HashKey, Key.y );
 	HashKey = MurmurAdd( HashKey, Key.z );
-	HashKey = MurmurAdd( HashKey, DetailLevelSpacing);
+	HashKey = MurmurAdd( HashKey, NormalDirectionBits );
 	uint Index;
 	if (ICHashTableFind(HashKey, Index))
 	{
 		uint RecordIndex = IrradianceCachingParameters.RWHashToIndex[Index];
-		FIrradianceCacheRecord Record = IrradianceCachingParameters.IrradianceCacheRecords[RecordIndex];
-		if (Record.Irradiance.w > IrradianceCachingParameters.Quality)
+		uint4 RecordIrradianceAndSampleCount = IrradianceCachingParameters.IrradianceCacheRecords[RecordIndex];
+		if (RecordIrradianceAndSampleCount.w > 0)
 		{
 			bIrradianceQuerySuccessful = true;
 			
-			RecordIrradiance = asfloat(Record.Irradiance.xyz) / Record.Irradiance.w;
+			RecordIrradiance = asfloat(RecordIrradianceAndSampleCount.xyz) / RecordIrradianceAndSampleCount.w;
 		}
 		
 		bGeometryQuerySuccessful = true;
-		RecordWorldNormal = Record.WorldNormal.xyz;
 		NearestRecordIndex = RecordIndex;
 	}
-	else
-	{
-		Key = EncodeVoxelKey(floor(WorldPosition / CoarseLevelSpacing));	
-		HashKey = 0;
-		HashKey = MurmurAdd( HashKey, Key.x );
-		HashKey = MurmurAdd( HashKey, Key.y );
-		HashKey = MurmurAdd( HashKey, Key.z );
-		HashKey = MurmurAdd( HashKey, NormalDirectionBits );
-		HashKey = MurmurAdd( HashKey, CoarseLevelSpacing);
-		
-		if (ICHashTableFind(HashKey, Index))
-		{
-			uint RecordIndex = IrradianceCachingParameters.RWHashToIndex[Index];
-			FIrradianceCacheRecord Record = IrradianceCachingParameters.IrradianceCacheRecords[RecordIndex];
-			if (Record.Irradiance.w > IrradianceCachingParameters.Quality)
-			{
-				bIrradianceQuerySuccessful = true;
-				
-				RecordIrradiance = asfloat(Record.Irradiance.xyz) / Record.Irradiance.w;
-			}
-			
-			bGeometryQuerySuccessful = true;
-			RecordWorldNormal = Record.WorldNormal.xyz;
-			NearestRecordIndex = RecordIndex;
-		}
-	}
 	
 	if (bIrradianceQuerySuccessful)
 	{		
diff --git a/Engine/Plugins/Experimental/GPULightmass/Shaders/Private/IrradianceCachingCommon.ush b/Engine/Plugins/Experimental/GPULightmass/Shaders/Private/IrradianceCachingCommon.ush
--- a/Engine/Plugins/Experimental/GPULightmass/Shaders/Private/IrradianceCachingCommon.ush
+++ b/Engine/Plugins/Experimental/GPULightmass/Shaders/Private/IrradianceCachingCommon.ush
@@ -1,5 +1,11 @@
 #include "/Engine/Private/HashTable.ush"
 
+struct FFinalGatherHitPoint
+{
+	float3 WorldPosition;
+	float3 WorldNormal;
+};
+
 bool ICHashTableAdd( uint Key, out uint Index )
 {
 	// Zero is reserved as invalid
@@ -74,6 +80,34 @@ uint3 EncodeVoxelKey(float3 VoxelPos)
 	return Key;
 }
 
+float3 EncodePositionKey(float3 WorldPosition, float3 WorldNormal, float Spacing)
+{
+	// Our cache entries are square shaped surfels. Along the primary axis of the normal the length is 1/4 of the specified cell size - which means 4x resolution.
+	// This reduces leaks through thin walls and also serves as a fail-safe plan when the geometry normal is not reliable
+	
+	const int CompressionFactor = 4;
+	
+	if(abs(WorldNormal.x) >= abs(WorldNormal.y) && abs(WorldNormal.x) >= abs(WorldNormal.z))
+	{
+		WorldPosition.x *= CompressionFactor;
+		return floor(WorldPosition / Spacing);
+	}
+	
+	if(abs(WorldNormal.y) >= abs(WorldNormal.x) && abs(WorldNormal.y) >= abs(WorldNormal.z))
+	{
+		WorldPosition.y *= CompressionFactor;
+		return floor(WorldPosition / Spacing);
+	}
+	
+	if(abs(WorldNormal.z) >= abs(WorldNormal.x) && abs(WorldNormal.z) >= abs(WorldNormal.y))
+	{				
+		WorldPosition.z *= CompressionFactor;
+		return floor(WorldPosition / Spacing);
+	}
+	
+	return floor(WorldPosition / Spacing);
+}
+
 uint EncodeNormalBits(float3 WorldNormal)
 {
 	return	(WorldNormal.x >= 0.01 ? (1 << 5) : 0) |
@@ -101,17 +135,16 @@ uint EncodeNormalBits(float3 WorldNormal)
 	} \
 }
 
-void EmitGeometryHitPoint(FIrradianceCacheRecord HitPoint, uint Size)
+void EmitGeometryHitPoint(FFinalGatherHitPoint HitPoint, uint Size)
 {
-	uint3 Key = EncodeVoxelKey(floor(HitPoint.WorldPosition.xyz / Size));
-	uint NormalDirectionBits = EncodeNormalBits(HitPoint.WorldNormal.xyz);
+	uint3 Key = EncodeVoxelKey(EncodePositionKey(HitPoint.WorldPosition, HitPoint.WorldNormal, Size));
+	uint NormalDirectionBits = EncodeNormalBits(HitPoint.WorldNormal);
 	
 	uint HashKey = 0;
 	HashKey = MurmurAdd( HashKey, Key.x );
 	HashKey = MurmurAdd( HashKey, Key.y );
 	HashKey = MurmurAdd( HashKey, Key.z );
 	HashKey = MurmurAdd( HashKey, NormalDirectionBits );
-	HashKey = MurmurAdd( HashKey, Size );
 	uint Index;
 	if (ICHashTableAdd(HashKey, Index))
 	{
@@ -124,6 +157,15 @@ void EmitGeometryHitPoint(FIrradianceCacheRecord HitPoint, uint Size)
 			IrradianceCachingParameters.RWHashTable[IrradianceCachingParameters.RWIndexToHash[RecordIndex]] = 0;
 		}
 		IrradianceCachingParameters.RWIndexToHash[RecordIndex] = Index;
-		IrradianceCachingParameters.IrradianceCacheRecords[RecordIndex] = HitPoint;
+		IrradianceCachingParameters.IrradianceCacheRecords[RecordIndex] = uint4(0, 0, 0, 0);
 	}
 }
+
+int GetIrradianceCachingQuality()
+{
+#ifdef VOLUMETRIC_LIGHTMAP_PATH_TRACING_MAIN_RG
+	return IrradianceCachingParameters.Quality * VolumetricLightmapQualityMultiplier;
+#endif
+	
+	return IrradianceCachingParameters.Quality;
+}
diff --git a/Engine/Plugins/Experimental/GPULightmass/Shaders/Private/LightmapOutput.usf b/Engine/Plugins/Experimental/GPULightmass/Shaders/Private/LightmapOutput.usf
--- a/Engine/Plugins/Experimental/GPULightmass/Shaders/Private/LightmapOutput.usf
+++ b/Engine/Plugins/Experimental/GPULightmass/Shaders/Private/LightmapOutput.usf
@@ -16,6 +16,7 @@ RWTexture2D<float4> ShadowMask;
 RWTexture2D<float4> ShadowMaskSampleCount;
 
 int NumTotalSamples;
+int NumIrradianceCachePasses;
 int NumRayGuidingTrialSamples;
 
 [numthreads(8, 8, 1)]
@@ -68,11 +69,16 @@ void SelectiveLightmapOutputCS(uint3 DispatchThreadId : SV_DispatchThreadID)
 					
 					float3 Color = float3(0, 1, 0);
 					
-					if (BatchedTiles[TileIndex].RenderPassIndex < NumRayGuidingTrialSamples)
+					if (BatchedTiles[TileIndex].RenderPassIndex < NumIrradianceCachePasses)
 					{
 						Color = float3(1, 0, 0);
 					}
 					
+					if (BatchedTiles[TileIndex].RenderPassIndex >= NumIrradianceCachePasses && BatchedTiles[TileIndex].RenderPassIndex < NumIrradianceCachePasses + NumRayGuidingTrialSamples)
+					{
+						Color = float3(0.5, 0.5, 0);
+					}
+					
 					OutputTileAtlas[LaunchIndex + BatchedTiles[TileIndex].OutputLayer0Position].rgba = float4(
 					Color, log2( 1 + LogBlackPoint ) - (1.0f / 255 - 0.5 / 255));
 				}
diff --git a/Engine/Plugins/Experimental/GPULightmass/Shaders/Private/LightmapPathTracing.usf b/Engine/Plugins/Experimental/GPULightmass/Shaders/Private/LightmapPathTracing.usf
--- a/Engine/Plugins/Experimental/GPULightmass/Shaders/Private/LightmapPathTracing.usf
+++ b/Engine/Plugins/Experimental/GPULightmass/Shaders/Private/LightmapPathTracing.usf
@@ -8,14 +8,6 @@ PathTracingRayGenShader.usf: Reference path tracing
 #define RANDSEQ_UNROLL_SOBOL 0
 #define USE_PATH_TRACING_LIGHT_GRID 1
 
-// Needs to be defined before the uniform buffer struct header is included, very unfortunately
-struct FIrradianceCacheRecord
-{
-	float4 WorldPosition;
-	float4 WorldNormal;
-	uint4 Irradiance;
-};
-
 #include "/Engine/Private/Common.ush"
 #include "/Engine/Private/PostProcessCommon.ush"
 #include "/Engine/Private/RectLight.ush"
@@ -38,13 +30,17 @@ RWTexture2D<uint> RayGuidingLuminance;
 Texture2D<float> RayGuidingCDFX;
 Texture2D<float> RayGuidingCDFY;
 
+#ifdef VOLUMETRIC_LIGHTMAP_PATH_TRACING_MAIN_RG
+int VolumetricLightmapQualityMultiplier;
+#endif
+
+#include "IrradianceCachingCommon.ush"
+
 #include "/Engine/Private/ShadingModels.ush"
 #include "/Engine/Private/PathTracing/Utilities/PathTracingRandomSequence.ush"
 #include "/Engine/Private/PathTracing/Light/PathTracingLightSampling.ush"
 #include "/Engine/Private/PathTracing/Material/PathTracingMaterialSampling.ush"
 
-#include "IrradianceCachingCommon.ush"
-
 #include "LightmapEncoding.ush"
 
 int LastInvalidationFrame;
@@ -329,7 +325,7 @@ void PathTracingKernel(
 #if USE_IRRADIANCE_CACHING
 	bool bShouldEmitGeometryHitPoint = false;
 	
-	FIrradianceCacheRecord IrradianceCachePrimaryGIHitPoint = (FIrradianceCacheRecord)0;	
+	FFinalGatherHitPoint IrradianceCacheFinalGatherHitPoint = (FFinalGatherHitPoint)0;	
 	float RadianceProbe_Pdf = 1.0f / (2 * PI);
 	
 	uint NewEntrySize = 32;
@@ -409,46 +405,77 @@ void PathTracingKernel(
 		}
 
 	#if USE_IRRADIANCE_CACHING
-		uint CoarseLevelSpacing = IrradianceCachingParameters.Spacing;
-		uint DetailLevelSpacing = IrradianceCachingParameters.Spacing;
+		float3 TexelPosition = TranslatedWorldPosition - LWCHackToFloat(PrimaryView.PreViewTranslation);
+		float3 HitPosition = Payload.TranslatedWorldPos - LWCHackToFloat(PrimaryView.PreViewTranslation);
+		float DistanceFromTexelToHit = length(HitPosition - TexelPosition);
+		bool bRejectedByCornerRejection = DistanceFromTexelToHit < IrradianceCachingParameters.Spacing * IrradianceCachingParameters.CornerRejection;
+			
+		// If we're in IC building passes and the entry is going to be rejected, terminate immediately as the path isn't going to generate any useful result
+		if (RenderPassIndex < GetIrradianceCachingQuality())
+		{
+			if (Bounce == 1 && bRejectedByCornerRejection)
+			{
+				break;
+			}
+		}
 		
-		if (NearestCacheEntryID == -1 && Bounce >= 1 && Payload.HitT >= DetailLevelSpacing * IrradianceCachingParameters.CornerRejection)
+		if (NearestCacheEntryID == -1 && Bounce >= 1 && !bRejectedByCornerRejection)
 		{
 			bool bIrradianceQuerySuccessful = false;
 			bool bGeometryQuerySuccessful = false;
 			
 			uint NearestRecordIndex = 0;
 			float3 RecordIrradiance;
-			float3 RecordWorldNormal;
 			
-			uint3 Key = EncodeVoxelKey(floor(Payload.TranslatedWorldPos / (Payload.HitT >= CoarseLevelSpacing ? CoarseLevelSpacing : DetailLevelSpacing)));
+			float3 PosKey = EncodePositionKey(HitPosition, Payload.WorldNormal, IrradianceCachingParameters.Spacing); // RT_LWC_TODO
+			uint3 Key = EncodeVoxelKey(PosKey);
 			uint NormalDirectionBits = EncodeNormalBits(Payload.WorldNormal);			
-			uint HashKey = 0;
+			uint HashKey = 0;		
 			HashKey = MurmurAdd( HashKey, Key.x );
 			HashKey = MurmurAdd( HashKey, Key.y );
 			HashKey = MurmurAdd( HashKey, Key.z );
-			HashKey = MurmurAdd( HashKey, NormalDirectionBits );
-			HashKey = MurmurAdd( HashKey, Payload.HitT >= CoarseLevelSpacing ? CoarseLevelSpacing : DetailLevelSpacing );
+			HashKey = MurmurAdd( HashKey, NormalDirectionBits );	
 			uint Index;
 			if (ICHashTableFind(HashKey, Index))
 			{
 				uint RecordIndex = IrradianceCachingParameters.RWHashToIndex[Index];
-				FIrradianceCacheRecord Record = IrradianceCachingParameters.IrradianceCacheRecords[RecordIndex];
-				if (Record.Irradiance.w > IrradianceCachingParameters.Quality)
+
+				uint4 RecordIrradianceAndSampleCount = IrradianceCachingParameters.IrradianceCacheRecords[RecordIndex];
+
+				// Compute an overestimated upper bound for the number of samples we need for this texel by assuming
+				// all the rays from all IC building passes are sent from this single texel
+				float CellArea = IrradianceCachingParameters.Spacing * IrradianceCachingParameters.Spacing;
+				float TotalSamplesOverAllICPasses = GPreviewLightmapPhysicalTileSize * GPreviewLightmapPhysicalTileSize * GetIrradianceCachingQuality();
+				float DistanceFraction = CellArea / (2 * PI * DistanceFromTexelToHit * DistanceFromTexelToHit);
+			
+				// We accept a cache entry, if:
+				// 1) We're outside of IC building passes. An entry with low sample count means we just need that much - still, we're going to require a minimum of 8 samples to remove outliers
+				// 2) The entry has been built by other lightmap tiles and contains more samples than we will ever need.
+
+				bool bShouldAcceptCacheEntry = false;
+
+			#ifdef LIGHTMAP_PATH_TRACING_MAIN_RG
+				// Surface lightmaps
+				bShouldAcceptCacheEntry = RenderPassIndex >= GetIrradianceCachingQuality() || RecordIrradianceAndSampleCount.w >= TotalSamplesOverAllICPasses * DistanceFraction;
+			#else
+				// Volumetric lightmaps
+				bShouldAcceptCacheEntry = RenderPassIndex >= GetIrradianceCachingQuality();
+			#endif
+			
+				if (bShouldAcceptCacheEntry)
 				{
 					bIrradianceQuerySuccessful = true;
-					
-					RecordIrradiance = asfloat(Record.Irradiance.xyz) / Record.Irradiance.w;
+		
+					RecordIrradiance = asfloat(RecordIrradianceAndSampleCount.xyz) / RecordIrradianceAndSampleCount.w;
 				}
 				
 				bGeometryQuerySuccessful = true;
-				RecordWorldNormal = Record.WorldNormal.xyz;
 				NearestRecordIndex = RecordIndex;
 			}
 	
 			if (bIrradianceQuerySuccessful)
 			{
-				// Successful query. Terminate path immediately. Can happen when accumulating for IC entries.
+				// Successful query. Terminate path immediately.
 				Radiance += RecordIrradiance * PathThroughput;
 				break;
 			}
@@ -459,10 +486,24 @@ void PathTracingKernel(
 				{
 					if (!bGeometryQuerySuccessful)
 					{
-						bShouldEmitGeometryHitPoint = true;
-						IrradianceCachePrimaryGIHitPoint.WorldPosition.xyz = Payload.TranslatedWorldPos - LWCHackToFloat(PrimaryView.PreViewTranslation); // RT_LWC_TODO
-						IrradianceCachePrimaryGIHitPoint.WorldNormal.xyz = Payload.WorldNormal;
-						NewEntrySize = Payload.HitT >= CoarseLevelSpacing ? CoarseLevelSpacing : DetailLevelSpacing;
+						bool bShouldPlaceNewEntry = false;
+					
+					#ifdef LIGHTMAP_PATH_TRACING_MAIN_RG
+						// Surface lightmaps are relatively high density. Only use the very first few passes to place new entries to reduce cache thrashing
+						bShouldPlaceNewEntry = RenderPassIndex < 16;
+					#else
+						// Always place new entry for volumetric lightmaps
+						bShouldPlaceNewEntry = true;
+					#endif
+						
+						if (bShouldPlaceNewEntry)
+						{
+							bShouldEmitGeometryHitPoint = true;
+							IrradianceCacheFinalGatherHitPoint.WorldPosition = HitPosition;
+							IrradianceCacheFinalGatherHitPoint.WorldNormal = Payload.WorldNormal;
+							NewEntrySize = IrradianceCachingParameters.Spacing;
+							break;
+						}
 					}
 					else
 					{
@@ -716,16 +757,16 @@ void PathTracingKernel(
 	{
 		if(bShouldEmitGeometryHitPoint)
 		{
-			EmitGeometryHitPoint(IrradianceCachePrimaryGIHitPoint, NewEntrySize);
+			EmitGeometryHitPoint(IrradianceCacheFinalGatherHitPoint, NewEntrySize);
 		}
 		
 		if(NearestCacheEntryID != -1)
 		{
 			float3 IrradianceToAccumulate = Radiance * RadianceProbe_Pdf; // Revert the radiance probe pdf as we're caching the first bounce outgoing radiance
-			ATOMIC_ADD_FLOAT(IrradianceCachingParameters.IrradianceCacheRecords[NearestCacheEntryID].Irradiance.r, IrradianceToAccumulate.r);
-			ATOMIC_ADD_FLOAT(IrradianceCachingParameters.IrradianceCacheRecords[NearestCacheEntryID].Irradiance.g, IrradianceToAccumulate.g);
-			ATOMIC_ADD_FLOAT(IrradianceCachingParameters.IrradianceCacheRecords[NearestCacheEntryID].Irradiance.b, IrradianceToAccumulate.b);
-			InterlockedAdd(IrradianceCachingParameters.IrradianceCacheRecords[NearestCacheEntryID].Irradiance.a, 1);
+			ATOMIC_ADD_FLOAT(IrradianceCachingParameters.IrradianceCacheRecords[NearestCacheEntryID].r, IrradianceToAccumulate.r);
+			ATOMIC_ADD_FLOAT(IrradianceCachingParameters.IrradianceCacheRecords[NearestCacheEntryID].g, IrradianceToAccumulate.g);
+			ATOMIC_ADD_FLOAT(IrradianceCachingParameters.IrradianceCacheRecords[NearestCacheEntryID].b, IrradianceToAccumulate.b);
+			InterlockedAdd(IrradianceCachingParameters.IrradianceCacheRecords[NearestCacheEntryID].a, 1);
 		}
 	}
 #endif
@@ -756,6 +797,28 @@ void LightmapPathTracingMainRG()
 
 	float3 TranslatedWorldPosition = EncodedGBufferWorldPosition.xyz / WorldPositionScalar + LWCHackToFloat(PrimaryView.PreViewTranslation); // RT_LWC_TODO
 	
+	int EffectiveRenderPassIndex = BatchedTiles[TileIndex].RenderPassIndex;
+#if USE_IRRADIANCE_CACHING
+	if (EffectiveRenderPassIndex >= GetIrradianceCachingQuality())
+	{
+		EffectiveRenderPassIndex -= GetIrradianceCachingQuality();
+
+	#if USE_FIRST_BOUNCE_RAY_GUIDING
+		if (EffectiveRenderPassIndex >= NumRayGuidingTrialSamples)
+		{
+			EffectiveRenderPassIndex -= NumRayGuidingTrialSamples;
+		}
+	#endif
+	}
+#endif
+					
+	if (EffectiveRenderPassIndex == 0) 
+	{
+		IrradianceAndSampleCount[TexelIndexInPool]= float4(0, 0, 0, 0);
+		SHDirectionality[TexelIndexInPool] = float4(0, 0, 0, 0);
+		SHCorrectionAndStationarySkyLightBentNormal[TexelIndexInPool].x = 0;
+	}
+	
 	if (NumTotalSamples <= 0 || (NumTotalSamples > 0 && BatchedTiles[TileIndex].RenderPassIndex < NumTotalSamples - 1))
 	{
 		if (bIsValidSample)
@@ -768,13 +831,7 @@ void LightmapPathTracingMainRG()
 			
 			{
 				RandomSequence RandSequence;
-				RandomSequence_Initialize(RandSequence, Seed, 
-				#if USE_FIRST_BOUNCE_RAY_GUIDING
-					BatchedTiles[TileIndex].RenderPassIndex >= NumRayGuidingTrialSamples ? 
-						(BatchedTiles[TileIndex].RenderPassIndex - NumRayGuidingTrialSamples) : 
-				#endif
-						(BatchedTiles[TileIndex].RenderPassIndex)
-				);
+				RandomSequence_Initialize(RandSequence, Seed, EffectiveRenderPassIndex);
 				
 				uint SampleIndex = 2;
 
@@ -807,25 +864,27 @@ void LightmapPathTracingMainRG()
 				
 				if (bIsValidSample)
 				{
-					#if USE_FIRST_BOUNCE_RAY_GUIDING						
-						if (BatchedTiles[TileIndex].RenderPassIndex < NumRayGuidingTrialSamples)
-						{
-							int2 ClusterPosition = clamp((int2)LaunchIndex - int2(2, 2), int2(0, 0), int2(63, 63)) / TEXEL_CLUSTER_SIZE;
-							float2 JitteredBin = PrimaryRandSample.xy * DIRECTIONAL_BINS_ONE_DIM;// - float2(0.5f, 0.5f) + PrimaryRandSample.zw;
-							int2 PositionInBin = clamp(JitteredBin, float2(0, 0), float2(DIRECTIONAL_BINS_ONE_DIM - 1, DIRECTIONAL_BINS_ONE_DIM - 1));
-							int2 FinalPosition = ClusterPosition * DIRECTIONAL_BINS_ONE_DIM + PositionInBin;
-							float Illuminance = Luminance(RadianceValue) * saturate(dot(RadianceDirection, ShadingNormal));
-							// For positive floats we can cast them to uint and use atomic max directly
-							InterlockedMax(RayGuidingLuminance[BatchedTiles[TileIndex].WorkingSetPosition / GPreviewLightmapPhysicalTileSize * CDF_TILE_SIZE + FinalPosition], asuint(max(Illuminance, 0)));
-						}
+				#if USE_FIRST_BOUNCE_RAY_GUIDING
+					int MinRenderPassIndex = 0;
+					int MaxRenderPassIndex = NumRayGuidingTrialSamples;
+				
+				#if USE_IRRADIANCE_CACHING
+					MinRenderPassIndex += GetIrradianceCachingQuality();
+					MaxRenderPassIndex += GetIrradianceCachingQuality();
+				#endif
 					
-						if (BatchedTiles[TileIndex].RenderPassIndex == NumRayGuidingTrialSamples) 
-						{
-							IrradianceAndSampleCount[TexelIndexInPool]= float4(0, 0, 0, 0);
-							SHDirectionality[TexelIndexInPool] = float4(0, 0, 0, 0);
-							SHCorrectionAndStationarySkyLightBentNormal[TexelIndexInPool].x = 0;
-						}
-					#endif
+					if (BatchedTiles[TileIndex].RenderPassIndex >= MinRenderPassIndex && BatchedTiles[TileIndex].RenderPassIndex < MaxRenderPassIndex)
+					{
+						int2 ClusterPosition = clamp((int2)LaunchIndex - int2(2, 2), int2(0, 0), int2(63, 63)) / TEXEL_CLUSTER_SIZE;
+						float2 JitteredBin = PrimaryRandSample.xy * DIRECTIONAL_BINS_ONE_DIM;// - float2(0.5f, 0.5f) + PrimaryRandSample.zw;
+						int2 PositionInBin = clamp(JitteredBin, float2(0, 0), float2(DIRECTIONAL_BINS_ONE_DIM - 1, DIRECTIONAL_BINS_ONE_DIM - 1));
+						int2 FinalPosition = ClusterPosition * DIRECTIONAL_BINS_ONE_DIM + PositionInBin;
+					
+						float Illuminance = Luminance(RadianceValue) * saturate(dot(RadianceDirection, ShadingNormal));
+						// For positive floats we can cast them to uint and use atomic max directly
+						InterlockedMax(RayGuidingLuminance[BatchedTiles[TileIndex].WorkingSetPosition / GPreviewLightmapPhysicalTileSize * CDF_TILE_SIZE + FinalPosition], asuint(max(Illuminance, 0)));
+					}
+				#endif
 					
 					{
 						FL2SHAndCorrection SH = (FL2SHAndCorrection)0;
@@ -1068,13 +1127,21 @@ void VolumetricLightmapPathTracingMainRG()
 	int3 CellPosInVLM = (BrickRequests[BrickIndex].xyz * 4 + CellPosInBrick) * BrickRequests[BrickIndex].w;	
 	uint Seed = ComputeBrickLinearAddress(CellPosInVLM, IndirectionTextureDim * 4);
 	
+	int EffectiveRenderPassIndex = FrameNumber;
+#if USE_IRRADIANCE_CACHING
+	if (EffectiveRenderPassIndex >= GetIrradianceCachingQuality())
+	{
+		EffectiveRenderPassIndex -= GetIrradianceCachingQuality();
+	}
+#endif
+
 	#if 0
 		float3 Jitter = RandomSequence_GenerateSample3D(RandSequence);
 	#else
 		float3 RandSample = float3(
-			Halton(MortonEncode3(CellPosInVLM) + FrameNumber / 2, 2), 
-			Halton(MortonEncode3(CellPosInVLM) + FrameNumber / 2, 3),
-			Halton(MortonEncode3(CellPosInVLM) + FrameNumber / 2, 5)
+			Halton(MortonEncode3(CellPosInVLM) + EffectiveRenderPassIndex / 2, 2), 
+			Halton(MortonEncode3(CellPosInVLM) + EffectiveRenderPassIndex / 2, 3),
+			Halton(MortonEncode3(CellPosInVLM) + EffectiveRenderPassIndex / 2, 5)
 		);
 		float3 Jitter = RandSample;
 	#endif		
@@ -1083,7 +1150,7 @@ void VolumetricLightmapPathTracingMainRG()
 	float3 DetailCellSize = VolumeSize.xyz / IndirectionTextureDim / 4;
 	float3 WorldPosition = VolumeMin.xyz + DetailCellSize * JitteredCellPosInVLM; // RT_LWC_TODO
 	float3 TranslatedWorldPosition = WorldPosition + LWCHackToFloat(PrimaryView.PreViewTranslation); // RT_LWC_TODO
-	float3 ShadingNormal = float3(0, 0, (FrameNumber % 2 == 0) ? 1 : -1);
+	float3 ShadingNormal = float3(0, 0, (EffectiveRenderPassIndex % 2 == 0) ? 1 : -1);
 
 	float3 RadianceValue = 0;
 	float3 RadianceDirection = 0;				
@@ -1094,7 +1161,7 @@ void VolumetricLightmapPathTracingMainRG()
 	if (bIsValidSample)
 	{
 		RandomSequence RandSequence;
-		RandomSequence_Initialize(RandSequence, Seed, FrameNumber / 2);	
+		RandomSequence_Initialize(RandSequence, Seed, EffectiveRenderPassIndex / 2);	
 		uint SampleIndex = 4;
 		
 		PathTracingKernel(
@@ -1123,8 +1190,8 @@ void VolumetricLightmapPathTracingMainRG()
 
 	uint PaddedBrickSize = 4 + 1;
 	uint3 VoxelPos = ComputeBrickLayoutPosition(BrickIndex, uint3(256, 256, 256)) * PaddedBrickSize + CellPosInBrick;
-
-	if (FrameNumber == 0)
+	
+	if (EffectiveRenderPassIndex == 0) 
 	{
 		AmbientVector[VoxelPos] = float4(0, 0, 0, 0);
 		
@@ -1200,6 +1267,11 @@ void VolumetricLightmapPathTracingMainRG()
 	// 32 samples for single sample stationary directional light shadowing
 	if (FrameNumber < 32 && length(LightShaderParametersArray[0].Direction) > 0)
 	{
+		if (FrameNumber == 0)
+		{
+			DirectionalLightShadowing[VoxelPos] = 0;
+		}
+		
 		FLightShaderParameters LightParameters;
 		LightShaderParametersArray[0].FillLightShaderParameters(LightParameters);
 
diff --git a/Engine/Plugins/Experimental/GPULightmass/Source/GPULightmass/Private/IrradianceCaching.cpp b/Engine/Plugins/Experimental/GPULightmass/Source/GPULightmass/Private/IrradianceCaching.cpp
--- a/Engine/Plugins/Experimental/GPULightmass/Source/GPULightmass/Private/IrradianceCaching.cpp
+++ b/Engine/Plugins/Experimental/GPULightmass/Source/GPULightmass/Private/IrradianceCaching.cpp
@@ -19,10 +19,10 @@ FIrradianceCache::FIrradianceCache(int32 Quality, float Spacing, float CornerRej
 
 	{
 		FRHIResourceCreateInfo CreateInfo(TEXT("FIrradianceCache"));
-		IrradianceCacheRecords = RHICreateStructuredBuffer(sizeof(FIrradianceCacheRecord), sizeof(FIrradianceCacheRecord) * IrradianceCacheMaxSize, BUF_UnorderedAccess | BUF_ShaderResource, CreateInfo);
+		IrradianceCacheRecords = RHICreateStructuredBuffer(sizeof(FVector4f), sizeof(FVector4f) * IrradianceCacheMaxSize, BUF_UnorderedAccess | BUF_ShaderResource, CreateInfo);
 		IrradianceCacheRecordsUAV = RHICreateUnorderedAccessView(IrradianceCacheRecords, false, false);
-
-		IrradianceCacheTotalBytes += sizeof(FIrradianceCacheRecord) * IrradianceCacheMaxSize;
+		
+		IrradianceCacheTotalBytes += sizeof(FVector4f) * IrradianceCacheMaxSize;
 
 		RHICmdList.ClearUAVUint(IrradianceCacheRecordsUAV, FUintVector4(0, 0, 0, 0));
 	}
diff --git a/Engine/Plugins/Experimental/GPULightmass/Source/GPULightmass/Private/IrradianceCaching.h b/Engine/Plugins/Experimental/GPULightmass/Source/GPULightmass/Private/IrradianceCaching.h
--- a/Engine/Plugins/Experimental/GPULightmass/Source/GPULightmass/Private/IrradianceCaching.h
+++ b/Engine/Plugins/Experimental/GPULightmass/Source/GPULightmass/Private/IrradianceCaching.h
@@ -12,7 +12,7 @@ BEGIN_UNIFORM_BUFFER_STRUCT(FIrradianceCachingParameters, )
 	SHADER_PARAMETER(int32, Quality)
 	SHADER_PARAMETER(float, Spacing)
 	SHADER_PARAMETER(float, CornerRejection)
-	SHADER_PARAMETER_UAV(RWStructuredBuffer<FIrradianceCacheRecord>, IrradianceCacheRecords)
+	SHADER_PARAMETER_UAV(RWStructuredBuffer<uint4>, IrradianceCacheRecords)
 	SHADER_PARAMETER_UAV(RWStructuredBuffer<uint>, RWHashTable)
 	SHADER_PARAMETER_UAV(RWStructuredBuffer<uint>, RWHashToIndex)
 	SHADER_PARAMETER_UAV(RWStructuredBuffer<uint>, RWIndexToHash)
@@ -21,17 +21,7 @@ END_UNIFORM_BUFFER_STRUCT()
 
 struct FIrradianceCache
 {
-	struct FIrradianceCacheRecord
-	{
-		// When used as a cache entry, 
-		// WorldPosition.w == FrameLastTouched
-		// WorldNormal.w == NumAccumulatedSamples
-		FVector4f WorldPosition;
-		FVector4f WorldNormal;
-		FVector4f Irradiance;
-	};
-
-	const int32 IrradianceCacheMaxSize = 1048576;
+	const int32 IrradianceCacheMaxSize = 4194304;
 
 	FBufferRHIRef IrradianceCacheRecords;
 	FUnorderedAccessViewRHIRef IrradianceCacheRecordsUAV;
diff --git a/Engine/Plugins/Experimental/GPULightmass/Source/GPULightmass/Private/LightmapRayTracing.h b/Engine/Plugins/Experimental/GPULightmass/Source/GPULightmass/Private/LightmapRayTracing.h
--- a/Engine/Plugins/Experimental/GPULightmass/Source/GPULightmass/Private/LightmapRayTracing.h
+++ b/Engine/Plugins/Experimental/GPULightmass/Source/GPULightmass/Private/LightmapRayTracing.h
@@ -107,6 +107,7 @@ class FVolumetricLightmapPathTracingRGS : public FGlobalShader
 	static void ModifyCompilationEnvironment(const FGlobalShaderPermutationParameters& Parameters, FShaderCompilerEnvironment& OutEnvironment)
 	{
 		OutEnvironment.SetDefine(TEXT("GPreviewLightmapPhysicalTileSize"), GPreviewLightmapPhysicalTileSize);
+		OutEnvironment.SetDefine(TEXT("VOLUMETRIC_LIGHTMAP_PATH_TRACING_MAIN_RG"), 1);
 		OutEnvironment.CompilerFlags.Add(CFLAG_ForceDXC);
 		OutEnvironment.CompilerFlags.Add(CFLAG_WarningsAsErrors);
 	}
@@ -118,6 +119,7 @@ class FVolumetricLightmapPathTracingRGS : public FGlobalShader
 		SHADER_PARAMETER(FIntVector, IndirectionTextureDim)
 		SHADER_PARAMETER(int32, NumTotalBricks)
 		SHADER_PARAMETER(int32, BrickBatchOffset)
+		SHADER_PARAMETER(int32, VolumetricLightmapQualityMultiplier)
 		SHADER_PARAMETER_SRV(RaytracingAccelerationStructure, TLAS)
 		SHADER_PARAMETER_SRV(Buffer<uint4>, BrickRequests)
 		SHADER_PARAMETER_UAV(RWTexture3D<float3>, AmbientVector)
@@ -207,7 +209,7 @@ class FFirstBounceRayGuidingCDFBuildCS : public FGlobalShader
 	}
 
 	BEGIN_SHADER_PARAMETER_STRUCT(FParameters, )
-		SHADER_PARAMETER(int, NumRayGuidingTrialSamples)
+		SHADER_PARAMETER(int, RayGuidingEndPassIndex)
 		SHADER_PARAMETER_SRV(StructuredBuffer<FGPUTileDescription>, BatchedTiles)
 		SHADER_PARAMETER_RDG_TEXTURE_UAV(RWTexture2D<uint>, RayGuidingLuminance)
 		SHADER_PARAMETER_RDG_TEXTURE_UAV(RWTexture2D<float>, RayGuidingCDFX)
diff --git a/Engine/Plugins/Experimental/GPULightmass/Source/GPULightmass/Private/LightmapRenderer.cpp b/Engine/Plugins/Experimental/GPULightmass/Source/GPULightmass/Private/LightmapRenderer.cpp
--- a/Engine/Plugins/Experimental/GPULightmass/Source/GPULightmass/Private/LightmapRenderer.cpp
+++ b/Engine/Plugins/Experimental/GPULightmass/Source/GPULightmass/Private/LightmapRenderer.cpp
@@ -106,6 +106,7 @@ class FSelectiveLightmapOutputCS : public FGlobalShader
 	BEGIN_SHADER_PARAMETER_STRUCT(FParameters, )
 		SHADER_PARAMETER(int32, NumBatchedTiles)
 		SHADER_PARAMETER(int32, NumTotalSamples)
+		SHADER_PARAMETER(int32, NumIrradianceCachePasses)
 		SHADER_PARAMETER(int32, NumRayGuidingTrialSamples)
 		SHADER_PARAMETER_SRV(StructuredBuffer<FGPUTileDescription>, BatchedTiles)
 		SHADER_PARAMETER_RDG_TEXTURE_UAV(RWTexture2D<float4>, OutputTileAtlas)
@@ -173,13 +174,19 @@ FLightmapRenderer::FLightmapRenderer(FSceneRenderState* InScene)
 	: Scene(InScene)
 	, LightmapTilePoolGPU(FIntPoint(Scene->Settings->LightmapTilePoolSize))
 {
-	bUseFirstBounceRayGuiding = Scene->Settings->bUseFirstBounceRayGuiding;
-	if (bUseFirstBounceRayGuiding)
+	NumTotalPassesToRender = Scene->Settings->GISamples;
+	
+	if (Scene->Settings->bUseIrradianceCaching)
 	{
-		NumFirstBounceRayGuidingTrialSamples = Scene->Settings->FirstBounceRayGuidingTrialSamples;
+		NumTotalPassesToRender += Scene->Settings->IrradianceCacheQuality;	
+	}
+	
+	if (Scene->Settings->bUseFirstBounceRayGuiding)
+	{
+		NumTotalPassesToRender += Scene->Settings->FirstBounceRayGuidingTrialSamples;
 	}
 
-	if (!bUseFirstBounceRayGuiding)
+	if (!Scene->Settings->bUseFirstBounceRayGuiding)
 	{
 		LightmapTilePoolGPU.Initialize(
 			{
@@ -1145,7 +1152,7 @@ bool FSceneRenderState::SetupRayTracingScene(int32 LODIndex)
 			TArray<FRHIRayTracingShader*> RayGenShaderTable;
 			{
 				FLightmapPathTracingRGS::FPermutationDomain PermutationVector;
-				PermutationVector.Set<FLightmapPathTracingRGS::FUseFirstBounceRayGuiding>(LightmapRenderer->bUseFirstBounceRayGuiding);
+				PermutationVector.Set<FLightmapPathTracingRGS::FUseFirstBounceRayGuiding>(Settings->bUseFirstBounceRayGuiding);
 				PermutationVector.Set<FLightmapPathTracingRGS::FUseIrradianceCaching>(Settings->bUseIrradianceCaching);
 				RayGenShaderTable.Add(GlobalShaderMap->GetShader<FLightmapPathTracingRGS>(PermutationVector).GetRayTracingShader());
 			}
@@ -2142,13 +2149,13 @@ void FLightmapRenderer::Finalize(FRDGBuilder& GraphBuilder)
 	}
 
 	int32 NumSamplesPerFrame = (bInsideBackgroundTick && bIsViewportNonRealtime) ? Scene->Settings->TilePassesInFullSpeedMode : Scene->Settings->TilePassesInSlowMode;
-	NumSamplesPerFrame = FMath::Max(FMath::Min(NumSamplesPerFrame, Scene->Settings->GISamples - 1), 0);
+	NumSamplesPerFrame = FMath::Max(FMath::Min(NumSamplesPerFrame, NumTotalPassesToRender - 1), 0);
 
 	{
 		auto& PendingGITileRequests = *GraphBuilder.AllocObject<TArray<FLightmapTileRequest>>(
-			PendingTileRequests.FilterByPredicate([NumGISamples = Scene->Settings->GISamples, MostCommonLODIndex](const FLightmapTileRequest& Tile)
+			PendingTileRequests.FilterByPredicate([NumTotalPassesToRender = NumTotalPassesToRender, MostCommonLODIndex](const FLightmapTileRequest& Tile)
 			{
-				return !Tile.RenderState->IsTileGIConverged(Tile.VirtualCoordinates, NumGISamples) && Tile.RenderState->GeometryInstanceRef.LODIndex == MostCommonLODIndex;
+				return !Tile.RenderState->IsTileGIConverged(Tile.VirtualCoordinates, NumTotalPassesToRender) && Tile.RenderState->GeometryInstanceRef.LODIndex == MostCommonLODIndex;
 			}));
 
 #if RHI_RAYTRACING
@@ -2218,7 +2225,7 @@ void FLightmapRenderer::Finalize(FRDGBuilder& GraphBuilder)
 								{
 									const FLightmapTileRequest& Tile = PendingGITileRequests[Index];
 
-									if (Tile.RenderState->IsTileGIConverged(Tile.VirtualCoordinates, Scene->Settings->GISamples)) continue;
+									if (Tile.RenderState->IsTileGIConverged(Tile.VirtualCoordinates, NumTotalPassesToRender)) continue;
 									uint32 AssignedGPUIndex = (Tile.RenderState->DistributionPrefixSum + Tile.RenderState->RetrieveTileStateIndex(Tile.VirtualCoordinates)) % GNumExplicitGPUsForRendering;
 									if (AssignedGPUIndex != GPUIndex) continue;
 
@@ -2233,13 +2240,31 @@ void FLightmapRenderer::Finalize(FRDGBuilder& GraphBuilder)
 
 									TArray<FMeshBatch> MeshBatches = Tile.RenderState->GeometryInstanceRef.GetMeshBatchesForGBufferRendering(Tile.VirtualCoordinates);
 
+									int32 EffectiveRenderPassIndex = PendingGIRenderPassIndices[Index];
+									
+									if(Scene->Settings->bUseIrradianceCaching)
+									{
+										if (EffectiveRenderPassIndex >= Scene->Settings->IrradianceCacheQuality)
+										{
+											EffectiveRenderPassIndex -= Scene->Settings->IrradianceCacheQuality;
+									
+											if(Scene->Settings->bUseFirstBounceRayGuiding)
+											{
+												if (EffectiveRenderPassIndex >= Scene->Settings->FirstBounceRayGuidingTrialSamples)
+												{
+													EffectiveRenderPassIndex -= Scene->Settings->FirstBounceRayGuidingTrialSamples;
+												}
+											}
+										}
+									}
+									
 									RenderMeshBatchesIntoGBuffer(
 										RHICmdList,
 										ReferenceView.Get(),
 										Scene->GetPrimitiveIdForGPUScene(Tile.RenderState->GeometryInstanceRef),
 										MeshBatches,
 										VirtualTexturePhysicalTileCoordinateScaleAndBias,
-										PendingGIRenderPassIndices[Index] / AAvsGIMultiplier,
+										EffectiveRenderPassIndex / AAvsGIMultiplier,
 										ScratchTilePoolGPU->GetPositionFromLinearAddress(Tile.TileAddressInScratch)* GPreviewLightmapPhysicalTileSize);
 
 									GPrimitiveIdVertexBufferPool.DiscardAll();
@@ -2271,7 +2296,7 @@ void FLightmapRenderer::Finalize(FRDGBuilder& GraphBuilder)
 							TileDesc.OutputLayer2Position = Tile.OutputPhysicalCoordinates[2] * GPreviewLightmapPhysicalTileSize;
 							TileDesc.FrameIndex = Tile.RenderState->RetrieveTileState(Tile.VirtualCoordinates).Revision;
 							TileDesc.RenderPassIndex = Tile.RenderState->RetrieveTileState(Tile.VirtualCoordinates).RenderPassIndex;
-							if (!Tile.RenderState->IsTileGIConverged(Tile.VirtualCoordinates, Scene->Settings->GISamples))
+							if (!Tile.RenderState->IsTileGIConverged(Tile.VirtualCoordinates, NumTotalPassesToRender))
 							{
 								Tile.RenderState->RetrieveTileState(Tile.VirtualCoordinates).RenderPassIndex++;
 
@@ -2312,7 +2337,7 @@ void FLightmapRenderer::Finalize(FRDGBuilder& GraphBuilder)
 							FRDGTextureRef RayGuidingCDFX = nullptr;
 							FRDGTextureRef RayGuidingCDFY = nullptr;
 
-							if (bUseFirstBounceRayGuiding)
+							if (Scene->Settings->bUseFirstBounceRayGuiding)
 							{
 								RayGuidingLuminance = GraphBuilder.RegisterExternalTexture(LightmapTilePoolGPU.PooledRenderTargets[5], TEXT("RayGuidingLuminance"));
 								RayGuidingCDFX = GraphBuilder.RegisterExternalTexture(LightmapTilePoolGPU.PooledRenderTargets[6], TEXT("RayGuidingCDFX"));
@@ -2328,7 +2353,7 @@ void FLightmapRenderer::Finalize(FRDGBuilder& GraphBuilder)
 								{
 									FLightmapPathTracingRGS::FParameters* PassParameters = GraphBuilder.AllocParameters<FLightmapPathTracingRGS::FParameters>();
 									PassParameters->LastInvalidationFrame = LastInvalidationFrame;
-									PassParameters->NumTotalSamples = Scene->Settings->GISamples;
+									PassParameters->NumTotalSamples = NumTotalPassesToRender;
 									PassParameters->TLAS = Scene->RayTracingSceneSRV;
 									PassParameters->GBufferWorldPosition = GBufferWorldPosition;
 									PassParameters->GBufferWorldNormal = GBufferWorldNormal;
@@ -2337,12 +2362,12 @@ void FLightmapRenderer::Finalize(FRDGBuilder& GraphBuilder)
 									PassParameters->SHCorrectionAndStationarySkyLightBentNormal = GraphBuilder.CreateUAV(SHCorrectionAndStationarySkyLightBentNormal);
 									PassParameters->SHDirectionality = GraphBuilder.CreateUAV(SHDirectionality);
 
-									if (bUseFirstBounceRayGuiding)
+									if (Scene->Settings->bUseFirstBounceRayGuiding)
 									{
 										PassParameters->RayGuidingLuminance = GraphBuilder.CreateUAV(RayGuidingLuminance);
 										PassParameters->RayGuidingCDFX = RayGuidingCDFX;
 										PassParameters->RayGuidingCDFY = RayGuidingCDFY;
-										PassParameters->NumRayGuidingTrialSamples = NumFirstBounceRayGuidingTrialSamples;
+										PassParameters->NumRayGuidingTrialSamples = Scene->Settings->FirstBounceRayGuidingTrialSamples;
 									}
 
 									PassParameters->BatchedTiles = GPUBatchedTileRequests.BatchedTilesSRV;
@@ -2374,7 +2399,7 @@ void FLightmapRenderer::Finalize(FRDGBuilder& GraphBuilder)
 
 
 									FLightmapPathTracingRGS::FPermutationDomain PermutationVector;
-									PermutationVector.Set<FLightmapPathTracingRGS::FUseFirstBounceRayGuiding>(bUseFirstBounceRayGuiding);
+									PermutationVector.Set<FLightmapPathTracingRGS::FUseFirstBounceRayGuiding>(Scene->Settings->bUseFirstBounceRayGuiding);
 									PermutationVector.Set<FLightmapPathTracingRGS::FUseIrradianceCaching>(Scene->Settings->bUseIrradianceCaching);
 									auto RayGenerationShader = GlobalShaderMap->GetShader<FLightmapPathTracingRGS>(PermutationVector);
 									ClearUnusedGraphResources(RayGenerationShader, PassParameters);
@@ -2394,7 +2419,7 @@ void FLightmapRenderer::Finalize(FRDGBuilder& GraphBuilder)
 									});
 								}
 
-								if (bUseFirstBounceRayGuiding)
+								if (Scene->Settings->bUseFirstBounceRayGuiding)
 								{
 									FFirstBounceRayGuidingCDFBuildCS::FParameters* PassParameters = GraphBuilder.AllocParameters<FFirstBounceRayGuidingCDFBuildCS::FParameters>();
 
@@ -2402,7 +2427,12 @@ void FLightmapRenderer::Finalize(FRDGBuilder& GraphBuilder)
 									PassParameters->RayGuidingLuminance = GraphBuilder.CreateUAV(RayGuidingLuminance);
 									PassParameters->RayGuidingCDFX = GraphBuilder.CreateUAV(RayGuidingCDFX);
 									PassParameters->RayGuidingCDFY = GraphBuilder.CreateUAV(RayGuidingCDFY);
-									PassParameters->NumRayGuidingTrialSamples = NumFirstBounceRayGuidingTrialSamples;
+									PassParameters->RayGuidingEndPassIndex = Scene->Settings->FirstBounceRayGuidingTrialSamples - 1;
+
+									if (Scene->Settings->bUseIrradianceCaching)
+									{
+										PassParameters->RayGuidingEndPassIndex += Scene->Settings->IrradianceCacheQuality;
+									}
 
 									TShaderMapRef<FFirstBounceRayGuidingCDFBuildCS> ComputeShader(GlobalShaderMap);
 									FComputeShaderUtils::AddPass(
@@ -2785,7 +2815,7 @@ void FLightmapRenderer::Finalize(FRDGBuilder& GraphBuilder)
 				TransferTexture(3);
 				TransferTexture(4);
 
-				if (bUseFirstBounceRayGuiding)
+				if (Scene->Settings->bUseFirstBounceRayGuiding)
 				{
 					TransferTexture(5);
 					TransferTexture(6);
@@ -2851,8 +2881,9 @@ void FLightmapRenderer::Finalize(FRDGBuilder& GraphBuilder)
 
 				FSelectiveLightmapOutputCS::FParameters* PassParameters = GraphBuilder.AllocParameters<FSelectiveLightmapOutputCS::FParameters>();
 				PassParameters->NumBatchedTiles = GPUBatchedTileRequests.BatchedTilesDesc.Num();
-				PassParameters->NumTotalSamples = Scene->Settings->GISamples;
-				PassParameters->NumRayGuidingTrialSamples = NumFirstBounceRayGuidingTrialSamples;
+				PassParameters->NumTotalSamples = NumTotalPassesToRender;
+				PassParameters->NumIrradianceCachePasses = Scene->Settings->bUseIrradianceCaching ? Scene->Settings->IrradianceCacheQuality : 0;
+				PassParameters->NumRayGuidingTrialSamples = Scene->Settings->bUseFirstBounceRayGuiding ? Scene->Settings->FirstBounceRayGuidingTrialSamples : 0;
 				PassParameters->BatchedTiles = GPUBatchedTileRequests.BatchedTilesSRV;
 				PassParameters->OutputTileAtlas = GraphBuilder.CreateUAV(RenderTargetTileAtlas);
 				PassParameters->IrradianceAndSampleCount = GraphBuilder.CreateUAV(IrradianceAndSampleCount);
@@ -2879,7 +2910,7 @@ void FLightmapRenderer::Finalize(FRDGBuilder& GraphBuilder)
 
 				FSelectiveLightmapOutputCS::FParameters* PassParameters = GraphBuilder.AllocParameters<FSelectiveLightmapOutputCS::FParameters>();
 				PassParameters->NumBatchedTiles = GPUBatchedTileRequests.BatchedTilesDesc.Num();
-				PassParameters->NumTotalSamples = Scene->Settings->GISamples;
+				PassParameters->NumTotalSamples = NumTotalPassesToRender;
 				PassParameters->BatchedTiles = GPUBatchedTileRequests.BatchedTilesSRV;
 				PassParameters->OutputTileAtlas = GraphBuilder.CreateUAV(RenderTargetTileAtlas);
 				PassParameters->ShadowMask = GraphBuilder.CreateUAV(ShadowMask);
@@ -2904,7 +2935,7 @@ void FLightmapRenderer::Finalize(FRDGBuilder& GraphBuilder)
 	{
 		auto ConvergedTileRequests = PendingTileRequests.FilterByPredicate(
 			[
-				NumGISamples = Scene->Settings->GISamples,
+				NumGISamples = NumTotalPassesToRender,
 				NumShadowSamples = Scene->Settings->StationaryLightShadowSamples,
 				bOnlyBakeWhatYouSee = bOnlyBakeWhatYouSee, 
 				bDenoiseDuringInteractiveBake = bDenoiseDuringInteractiveBake
diff --git a/Engine/Plugins/Experimental/GPULightmass/Source/GPULightmass/Private/LightmapRenderer.h b/Engine/Plugins/Experimental/GPULightmass/Source/GPULightmass/Private/LightmapRenderer.h
--- a/Engine/Plugins/Experimental/GPULightmass/Source/GPULightmass/Private/LightmapRenderer.h
+++ b/Engine/Plugins/Experimental/GPULightmass/Source/GPULightmass/Private/LightmapRenderer.h
@@ -130,8 +130,7 @@ public:
 
 	int32 FrameNumber = 0;
 
-	bool bUseFirstBounceRayGuiding = false;
-	int32 NumFirstBounceRayGuidingTrialSamples = 0;
+	int32 NumTotalPassesToRender = 0;
 
 	bool bDenoiseDuringInteractiveBake = false;
 	bool bOnlyBakeWhatYouSee = false;
diff --git a/Engine/Plugins/Experimental/GPULightmass/Source/GPULightmass/Private/Scene/Scene.cpp b/Engine/Plugins/Experimental/GPULightmass/Source/GPULightmass/Private/Scene/Scene.cpp
--- a/Engine/Plugins/Experimental/GPULightmass/Source/GPULightmass/Private/Scene/Scene.cpp
+++ b/Engine/Plugins/Experimental/GPULightmass/Source/GPULightmass/Private/Scene/Scene.cpp
@@ -1644,10 +1644,10 @@ void FSceneRenderState::BackgroundTick()
 					{
 						FTileVirtualCoordinates VirtualCoordinates(FIntPoint(X, Y), 0);
 
-						TotalSamples += Settings->GISamples * GPreviewLightmapPhysicalTileSize * GPreviewLightmapPhysicalTileSize;
+						TotalSamples += LightmapRenderer->NumTotalPassesToRender * GPreviewLightmapPhysicalTileSize * GPreviewLightmapPhysicalTileSize;
 						SamplesTaken += (Lightmap.DoesTileHaveValidCPUData(VirtualCoordinates, LightmapRenderer->GetCurrentRevision()) ?
-							Settings->GISamples :
-							FMath::Min(Lightmap.RetrieveTileState(VirtualCoordinates).RenderPassIndex, Settings->GISamples - 1)) * GPreviewLightmapPhysicalTileSize * GPreviewLightmapPhysicalTileSize;
+							LightmapRenderer->NumTotalPassesToRender  :
+							FMath::Min(Lightmap.RetrieveTileState(VirtualCoordinates).RenderPassIndex, LightmapRenderer->NumTotalPassesToRender - 1)) * GPreviewLightmapPhysicalTileSize * GPreviewLightmapPhysicalTileSize;
 					}
 				}
 			}
@@ -1658,11 +1658,11 @@ void FSceneRenderState::BackgroundTick()
 			{
 				for (FLightmapTileRequest& Tile : LightmapRenderer->RecordedTileRequests)
 				{
-					TotalSamples += Settings->GISamples * GPreviewLightmapPhysicalTileSize * GPreviewLightmapPhysicalTileSize;
+					TotalSamples += LightmapRenderer->NumTotalPassesToRender * GPreviewLightmapPhysicalTileSize * GPreviewLightmapPhysicalTileSize;
 
 					SamplesTaken += (Tile.RenderState->DoesTileHaveValidCPUData(Tile.VirtualCoordinates, LightmapRenderer->GetCurrentRevision()) ?
-						Settings->GISamples :
-						FMath::Min(Tile.RenderState->RetrieveTileState(Tile.VirtualCoordinates).RenderPassIndex, Settings->GISamples - 1)) * GPreviewLightmapPhysicalTileSize * GPreviewLightmapPhysicalTileSize;
+						LightmapRenderer->NumTotalPassesToRender :
+						FMath::Min(Tile.RenderState->RetrieveTileState(Tile.VirtualCoordinates).RenderPassIndex, LightmapRenderer->NumTotalPassesToRender - 1)) * GPreviewLightmapPhysicalTileSize * GPreviewLightmapPhysicalTileSize;
 				}
 			}
 			else
@@ -1671,11 +1671,11 @@ void FSceneRenderState::BackgroundTick()
 				{
 					for (FLightmapTileRequest& Tile : FrameRequests)
 					{
-						TotalSamples += Settings->GISamples * GPreviewLightmapPhysicalTileSize * GPreviewLightmapPhysicalTileSize;
+						TotalSamples += LightmapRenderer->NumTotalPassesToRender  * GPreviewLightmapPhysicalTileSize * GPreviewLightmapPhysicalTileSize;
 
 						SamplesTaken += (Tile.RenderState->DoesTileHaveValidCPUData(Tile.VirtualCoordinates, LightmapRenderer->GetCurrentRevision()) ?
-							Settings->GISamples :
-							FMath::Min(Tile.RenderState->RetrieveTileState(Tile.VirtualCoordinates).RenderPassIndex, Settings->GISamples - 1)) * GPreviewLightmapPhysicalTileSize * GPreviewLightmapPhysicalTileSize;
+							LightmapRenderer->NumTotalPassesToRender :
+							FMath::Min(Tile.RenderState->RetrieveTileState(Tile.VirtualCoordinates).RenderPassIndex, LightmapRenderer->NumTotalPassesToRender - 1)) * GPreviewLightmapPhysicalTileSize * GPreviewLightmapPhysicalTileSize;
 					}
 				}
 			}
@@ -1686,7 +1686,7 @@ void FSceneRenderState::BackgroundTick()
 		{
 			int32 NumCellsPerBrick = 5 * 5 * 5;
 			SamplesTaken += VolumetricLightmapRenderer->SamplesTaken;
-			TotalSamples += (uint64)VolumetricLightmapRenderer->NumTotalBricks * NumCellsPerBrick * Settings->GISamples * VolumetricLightmapRenderer->GetGISamplesMultiplier();
+			TotalSamples += (uint64)VolumetricLightmapRenderer->NumTotalBricks * NumCellsPerBrick * VolumetricLightmapRenderer->NumTotalPassesToRender;
 		}
 
 		int32 IntPercentage = FMath::FloorToInt(SamplesTaken * 100.0 / TotalSamples);
diff --git a/Engine/Plugins/Experimental/GPULightmass/Source/GPULightmass/Private/VolumetricLightmap.cpp b/Engine/Plugins/Experimental/GPULightmass/Source/GPULightmass/Private/VolumetricLightmap.cpp
--- a/Engine/Plugins/Experimental/GPULightmass/Source/GPULightmass/Private/VolumetricLightmap.cpp
+++ b/Engine/Plugins/Experimental/GPULightmass/Source/GPULightmass/Private/VolumetricLightmap.cpp
@@ -85,6 +85,15 @@ FVolumetricLightmapRenderer::FVolumetricLightmapRenderer(FSceneRenderState* Scen
 	: Scene(Scene)
 {
 	VolumetricLightmap.Data = &VolumetricLightmapData;
+	
+	NumTotalPassesToRender = Scene->Settings->GISamples;
+	
+	if (Scene->Settings->bUseIrradianceCaching)
+	{
+		NumTotalPassesToRender += Scene->Settings->IrradianceCacheQuality;	
+	}
+
+	NumTotalPassesToRender *= Scene->Settings->VolumetricLightmapQualityMultiplier;
 }
 
 FPrecomputedVolumetricLightmap* FVolumetricLightmapRenderer::GetPrecomputedVolumetricLightmapForPreview()
@@ -469,11 +478,6 @@ void FVolumetricLightmapRenderer::VoxelizeScene()
 	Scene->DestroyRayTracingScene();
 }
 
-int32 FVolumetricLightmapRenderer::GetGISamplesMultiplier()
-{
-	return Scene->Settings->VolumetricLightmapQualityMultiplier;
-}
-
 void FVolumetricLightmapRenderer::BackgroundTick()
 {
 	if (NumTotalBricks == 0)
@@ -482,7 +486,7 @@ void FVolumetricLightmapRenderer::BackgroundTick()
 	}
 
 	int32 NumCellsPerBrick = 5 * 5 * 5;
-	if (SamplesTaken >= (uint64)NumTotalBricks * NumCellsPerBrick * Scene->Settings->GISamples * GetGISamplesMultiplier())
+	if (SamplesTaken >= (uint64)NumTotalBricks * NumCellsPerBrick * NumTotalPassesToRender)
 	{
 		return;
 	}
@@ -551,6 +555,7 @@ void FVolumetricLightmapRenderer::BackgroundTick()
 					PassParameters->BrickRequests = BrickRequests.SRV;
 					PassParameters->NumTotalBricks = NumTotalBricks;
 					PassParameters->BrickBatchOffset = BrickBatchOffset;
+					PassParameters->VolumetricLightmapQualityMultiplier = Scene->Settings->VolumetricLightmapQualityMultiplier;
 					PassParameters->AmbientVector = AccumulationBrickData.AmbientVector.UAV;
 					PassParameters->SHCoefficients0R = AccumulationBrickData.SHCoefficients[0].UAV;
 					PassParameters->SHCoefficients1R = AccumulationBrickData.SHCoefficients[1].UAV;
@@ -729,7 +734,7 @@ void FVolumetricLightmapRenderer::BackgroundTick()
 
 			SamplesTaken += BricksToCalcThisFrame * NumCellsPerBrick;
 
-			if (SamplesTaken >= (uint64)NumTotalBricks * NumCellsPerBrick * Scene->Settings->GISamples * GetGISamplesMultiplier())
+			if (SamplesTaken >= (uint64)NumTotalBricks * NumCellsPerBrick * NumTotalPassesToRender)
 			{
 				break;
 			}
diff --git a/Engine/Plugins/Experimental/GPULightmass/Source/GPULightmass/Private/VolumetricLightmap.h b/Engine/Plugins/Experimental/GPULightmass/Source/GPULightmass/Private/VolumetricLightmap.h
--- a/Engine/Plugins/Experimental/GPULightmass/Source/GPULightmass/Private/VolumetricLightmap.h
+++ b/Engine/Plugins/Experimental/GPULightmass/Source/GPULightmass/Private/VolumetricLightmap.h
@@ -27,10 +27,9 @@ public:
 	int32 NumTotalBricks;
 
 	int32 FrameNumber = 0;
+	int32 NumTotalPassesToRender = 0;
 	uint64 SamplesTaken = 0;
 
-	int32 GetGISamplesMultiplier();
-
 private:
 	FSceneRenderState* Scene;
 
diff --git a/Engine/Shaders/Private/PathTracing/Material/PathTracingRadianceProbe.ush b/Engine/Shaders/Private/PathTracing/Material/PathTracingRadianceProbe.ush
--- a/Engine/Shaders/Private/PathTracing/Material/PathTracingRadianceProbe.ush
+++ b/Engine/Shaders/Private/PathTracing/Material/PathTracingRadianceProbe.ush
@@ -64,7 +64,13 @@ FMaterialSample RadianceProbe_SampleMaterial(
 
 	int2 ClusterPosition = clamp((int2)LaunchIndex - int2(2, 2), int2(0, 0), int2(63, 63)) / TEXEL_CLUSTER_SIZE;
 	
-	if (BatchedTiles[TileIndex].RenderPassIndex >= NumRayGuidingTrialSamples)
+	int MinRenderPassIndex = NumRayGuidingTrialSamples;
+
+#if USE_IRRADIANCE_CACHING
+	MinRenderPassIndex += GetIrradianceCachingQuality();
+#endif
+	
+	if (BatchedTiles[TileIndex].RenderPassIndex >= MinRenderPassIndex)
 	{
 		float LastRowPrefixSum = 0;
 		float RowPrefixSum = 0;
@@ -148,8 +154,14 @@ float RadianceProbe_Pdf(float3 OutgoingDirection, FPathTracingPayload Payload)
 	int TileIndex = BatchedLaunchIndex.x / GPreviewLightmapPhysicalTileSize;
 
 	int2 ClusterPosition = clamp((int2)LaunchIndex - int2(2, 2), int2(0, 0), int2(63, 63)) / TEXEL_CLUSTER_SIZE;
+
+	int MinRenderPassIndex = NumRayGuidingTrialSamples;
+
+#if USE_IRRADIANCE_CACHING
+	MinRenderPassIndex += GetIrradianceCachingQuality();
+#endif
 	
-	if (BatchedTiles[TileIndex].RenderPassIndex >= NumRayGuidingTrialSamples)
+	if (BatchedTiles[TileIndex].RenderPassIndex >= MinRenderPassIndex)
 	{
 		float LastPrefixSum = 0;
 		float PrefixSum;
@@ -223,7 +235,13 @@ FMaterialEval RadianceProbe_EvalMaterial(
 
 	int2 ClusterPosition = clamp((int2)LaunchIndex - int2(2, 2), int2(0, 0), int2(63, 63)) / TEXEL_CLUSTER_SIZE;
 	
-	if (BatchedTiles[TileIndex].RenderPassIndex >= NumRayGuidingTrialSamples)
+	int MinRenderPassIndex = NumRayGuidingTrialSamples;
+
+#if USE_IRRADIANCE_CACHING
+	MinRenderPassIndex += GetIrradianceCachingQuality();
+#endif
+	
+	if (BatchedTiles[TileIndex].RenderPassIndex >= MinRenderPassIndex)
 	{
 		float LastPrefixSum = 0;
 		float PrefixSum;
