GPULM: Change IC hash function to reduce key mapping collision

* While the hash table has linear probing to resolve key collision, we're limited to 32bit keys (unless 64bit atomics are used) and that cannot represent the entire possible space that IC entries can be in.

diff --git a/Engine/Plugins/Experimental/GPULightmass/Shaders/Private/IrradianceCacheVisualization.usf b/Engine/Plugins/Experimental/GPULightmass/Shaders/Private/IrradianceCacheVisualization.usf
--- a/Engine/Plugins/Experimental/GPULightmass/Shaders/Private/IrradianceCacheVisualization.usf
+++ b/Engine/Plugins/Experimental/GPULightmass/Shaders/Private/IrradianceCacheVisualization.usf
@@ -25,14 +25,8 @@ void VisualizeIrradianceCachePS(
 			
 	uint NearestRecordIndex = 0;
 	float3 RecordIrradiance;
-	
-	uint3 Key = EncodeVoxelKey(EncodePositionKey(WorldPosition, ShadingNormal, IrradianceCachingParameters.Spacing));		
-	uint NormalDirectionBits = EncodeNormalBits(ShadingNormal);				
-	uint HashKey = 0;
-	HashKey = MurmurAdd( HashKey, Key.x );
-	HashKey = MurmurAdd( HashKey, Key.y );
-	HashKey = MurmurAdd( HashKey, Key.z );
-	HashKey = MurmurAdd( HashKey, NormalDirectionBits );
+			
+	uint HashKey = EncodeICHashKey(WorldPosition, ShadingNormal, IrradianceCachingParameters.Spacing);
 	uint Index;
 	if (ICHashTableFind(HashKey, Index))
 	{
diff --git a/Engine/Plugins/Experimental/GPULightmass/Shaders/Private/IrradianceCachingCommon.ush b/Engine/Plugins/Experimental/GPULightmass/Shaders/Private/IrradianceCachingCommon.ush
--- a/Engine/Plugins/Experimental/GPULightmass/Shaders/Private/IrradianceCachingCommon.ush
+++ b/Engine/Plugins/Experimental/GPULightmass/Shaders/Private/IrradianceCachingCommon.ush
@@ -119,6 +119,19 @@ uint EncodeNormalBits(float3 WorldNormal)
 			;
 }
 
+uint EncodeICHashKey(float3 WorldPosition, float3 WorldNormal, float Spacing)
+{
+	float3 PosKey = EncodePositionKey(WorldPosition, WorldNormal, Spacing); // RT_LWC_TODO
+	uint3 Key = EncodeVoxelKey(PosKey);
+	uint NormalDirectionBits = EncodeNormalBits(WorldNormal);			
+	uint HashKey = 0;		
+	HashKey = MurmurAdd( HashKey, NormalDirectionBits << 26);
+	HashKey = MurmurAdd( HashKey, Key.x << 17);
+	HashKey = MurmurAdd( HashKey, Key.y << 8);
+	HashKey = MurmurAdd( HashKey, Key.z << 0);
+	return HashKey;
+}
+
 #define ATOMIC_ADD_FLOAT(Value, Increment) \
 { \
 	uint NewValue = asuint(Increment); \
@@ -137,14 +150,7 @@ uint EncodeNormalBits(float3 WorldNormal)
 
 void EmitGeometryHitPoint(FFinalGatherHitPoint HitPoint, uint Size)
 {
-	uint3 Key = EncodeVoxelKey(EncodePositionKey(HitPoint.WorldPosition, HitPoint.WorldNormal, Size));
-	uint NormalDirectionBits = EncodeNormalBits(HitPoint.WorldNormal);
-	
-	uint HashKey = 0;
-	HashKey = MurmurAdd( HashKey, Key.x );
-	HashKey = MurmurAdd( HashKey, Key.y );
-	HashKey = MurmurAdd( HashKey, Key.z );
-	HashKey = MurmurAdd( HashKey, NormalDirectionBits );
+	uint HashKey = EncodeICHashKey(HitPoint.WorldPosition, HitPoint.WorldNormal, Size);
 	uint Index;
 	if (ICHashTableAdd(HashKey, Index))
 	{
diff --git a/Engine/Plugins/Experimental/GPULightmass/Shaders/Private/LightmapPathTracing.usf b/Engine/Plugins/Experimental/GPULightmass/Shaders/Private/LightmapPathTracing.usf
--- a/Engine/Plugins/Experimental/GPULightmass/Shaders/Private/LightmapPathTracing.usf
+++ b/Engine/Plugins/Experimental/GPULightmass/Shaders/Private/LightmapPathTracing.usf
@@ -426,15 +426,8 @@ void PathTracingKernel(
 			
 			uint NearestRecordIndex = 0;
 			float3 RecordIrradiance;
-			
-			float3 PosKey = EncodePositionKey(HitPosition, Payload.WorldNormal, IrradianceCachingParameters.Spacing); // RT_LWC_TODO
-			uint3 Key = EncodeVoxelKey(PosKey);
-			uint NormalDirectionBits = EncodeNormalBits(Payload.WorldNormal);			
-			uint HashKey = 0;		
-			HashKey = MurmurAdd( HashKey, Key.x );
-			HashKey = MurmurAdd( HashKey, Key.y );
-			HashKey = MurmurAdd( HashKey, Key.z );
-			HashKey = MurmurAdd( HashKey, NormalDirectionBits );	
+				
+			uint HashKey = EncodeICHashKey(HitPosition, Payload.WorldNormal, IrradianceCachingParameters.Spacing);
 			uint Index;
 			if (ICHashTableFind(HashKey, Index))
 			{
