GPULM: Optimize lightmap transcoding & encoding performance

* Sample quantization & copying is now parallel
* Oodle is used to compress temporary memory
* Now we have mimalloc editor-wide, we can remove FAnsiAllocator as its slow
* VT lightmaps are reduced from 32k back to 1-2k as huge single lightmaps hurt encoding performance seriously and 1-2k lightmaps should be adequate to reduce drawcalls

diff --git a/Engine/Plugins/Experimental/GPULightmass/Source/GPULightmass/Private/LightmapDenoising.cpp b/Engine/Plugins/Experimental/GPULightmass/Source/GPULightmass/Private/LightmapDenoising.cpp
--- a/Engine/Plugins/Experimental/GPULightmass/Source/GPULightmass/Private/LightmapDenoising.cpp
+++ b/Engine/Plugins/Experimental/GPULightmass/Source/GPULightmass/Private/LightmapDenoising.cpp
@@ -165,8 +165,8 @@ void DenoiseLightSampleData(FIntPoint Size, TArray<FLightSampleData>& LightSampl
 
 void DenoiseRawData(
 	FIntPoint Size,
-	TArray<FLinearColor, FAnsiAllocator>& IncidentLighting,
-	TArray<FLinearColor, FAnsiAllocator>& LuminanceSH,
+	TArray<FLinearColor>& IncidentLighting,
+	TArray<FLinearColor>& LuminanceSH,
 	FDenoiserContext& DenoiserContext, 
 	bool bPrepadTexels)
 {
diff --git a/Engine/Plugins/Experimental/GPULightmass/Source/GPULightmass/Private/LightmapDenoising.h b/Engine/Plugins/Experimental/GPULightmass/Source/GPULightmass/Private/LightmapDenoising.h
--- a/Engine/Plugins/Experimental/GPULightmass/Source/GPULightmass/Private/LightmapDenoising.h
+++ b/Engine/Plugins/Experimental/GPULightmass/Source/GPULightmass/Private/LightmapDenoising.h
@@ -92,7 +92,7 @@ void DenoiseLightSampleData(FIntPoint Size, TArray<FLightSampleData>& LightSampl
 
 void DenoiseRawData(
 	FIntPoint Size,
-	TArray<FLinearColor, FAnsiAllocator>& IncidentLighting,
-	TArray<FLinearColor, FAnsiAllocator>& LuminanceSH,
+	TArray<FLinearColor>& IncidentLighting,
+	TArray<FLinearColor>& LuminanceSH,
 	FDenoiserContext& DenoiserContext,
 	bool bPrepadTexels = true);
\ No newline at end of file
diff --git a/Engine/Plugins/Experimental/GPULightmass/Source/GPULightmass/Private/LightmapEncoding.cpp b/Engine/Plugins/Experimental/GPULightmass/Source/GPULightmass/Private/LightmapEncoding.cpp
--- a/Engine/Plugins/Experimental/GPULightmass/Source/GPULightmass/Private/LightmapEncoding.cpp
+++ b/Engine/Plugins/Experimental/GPULightmass/Source/GPULightmass/Private/LightmapEncoding.cpp
@@ -25,103 +25,152 @@ static void GetLUVW(const float RGB[3], float& L, float& U, float& V, float& W)
 	}
 }
 
+struct MinMaxCoefficients
+{
+	float MinCoefficient[NUM_STORED_LIGHTMAP_COEF][4];
+	float MaxCoefficient[NUM_STORED_LIGHTMAP_COEF][4];
+};
+
 void QuantizeLightSamples(
 	TArray<FLightSampleData> InLightSamples,
 	TArray<FLightMapCoefficients>& OutLightSamples,
 	float OutMultiply[NUM_STORED_LIGHTMAP_COEF][4],
 	float OutAdd[NUM_STORED_LIGHTMAP_COEF][4])
 {
+	FMemMark MemMark(FMemStack::Get());
+	
 	//const float EncodeExponent = .5f;
 	const float LogScale = 11.5f;
 	const float LogBlackPoint = FMath::Pow(2.0f, -0.5f * LogScale);
 	const float SimpleLogScale = 16.0f;
 	const float SimpleLogBlackPoint = FMath::Pow(2.0f, -0.5f * SimpleLogScale);
 
-	float MinCoefficient[NUM_STORED_LIGHTMAP_COEF][4];
-	float MaxCoefficient[NUM_STORED_LIGHTMAP_COEF][4];
-
-	for (int32 CoefficientIndex = 0; CoefficientIndex < NUM_STORED_LIGHTMAP_COEF; CoefficientIndex += 2)
+	const int32 NumSamplesPerTask = 4096;
+	const int32 NumTasks = FMath::DivideAndRoundUp(InLightSamples.Num(), NumSamplesPerTask);
+	TArray<MinMaxCoefficients, TMemStackAllocator<>> TaskStorage;
+	TaskStorage.AddUninitialized(NumTasks);
+	ParallelFor(NumTasks,[&](int32 TaskIndex)
 	{
-		for (int32 ColorIndex = 0; ColorIndex < 4; ColorIndex++)
+		float (&MinCoefficient)[NUM_STORED_LIGHTMAP_COEF][4] = TaskStorage[TaskIndex].MinCoefficient;
+		float (&MaxCoefficient)[NUM_STORED_LIGHTMAP_COEF][4] = TaskStorage[TaskIndex].MaxCoefficient;
+
+		for (int32 CoefficientIndex = 0; CoefficientIndex < NUM_STORED_LIGHTMAP_COEF; CoefficientIndex += 2)
 		{
-			// Color
-			MinCoefficient[CoefficientIndex][ColorIndex] = 10000.0f;
-			MaxCoefficient[CoefficientIndex][ColorIndex] = -10000.0f;
+			for (int32 ColorIndex = 0; ColorIndex < 4; ColorIndex++)
+			{
+				// Color
+				MinCoefficient[CoefficientIndex][ColorIndex] = 10000.0f;
+				MaxCoefficient[CoefficientIndex][ColorIndex] = -10000.0f;
 
-			// Direction
-			MinCoefficient[CoefficientIndex + 1][ColorIndex] = 10000.0f;
-			MaxCoefficient[CoefficientIndex + 1][ColorIndex] = -10000.0f;
+				// Direction
+				MinCoefficient[CoefficientIndex + 1][ColorIndex] = 10000.0f;
+				MaxCoefficient[CoefficientIndex + 1][ColorIndex] = -10000.0f;
+			}
 		}
-	}
 
-	// go over all samples looking for min and max values
-	for (int32 SampleIndex = 0; SampleIndex < InLightSamples.Num(); SampleIndex++)
-	{
-		const FLightSampleData& SourceSample = InLightSamples[SampleIndex];
-
-		if (SourceSample.bIsMapped)
+		// go over all samples looking for min and max values
+		for (int32 SampleIndex = TaskIndex * NumSamplesPerTask; SampleIndex < (TaskIndex + 1) * NumSamplesPerTask && SampleIndex < InLightSamples.Num(); SampleIndex++)
 		{
+			const FLightSampleData& SourceSample = InLightSamples[SampleIndex];
+
+			if (SourceSample.bIsMapped)
 			{
-				// Complex
-				float L, U, V, W;
-				GetLUVW(SourceSample.Coefficients[0], L, U, V, W);
+				{
+					// Complex
+					float L, U, V, W;
+					GetLUVW(SourceSample.Coefficients[0], L, U, V, W);
 
-				float LogL = FMath::Log2(L + LogBlackPoint);
+					float LogL = FMath::Log2(L + LogBlackPoint);
 
-				MinCoefficient[0][0] = FMath::Min(MinCoefficient[0][0], U);
-				MaxCoefficient[0][0] = FMath::Max(MaxCoefficient[0][0], U);
+					MinCoefficient[0][0] = FMath::Min(MinCoefficient[0][0], U);
+					MaxCoefficient[0][0] = FMath::Max(MaxCoefficient[0][0], U);
 
-				MinCoefficient[0][1] = FMath::Min(MinCoefficient[0][1], V);
-				MaxCoefficient[0][1] = FMath::Max(MaxCoefficient[0][1], V);
+					MinCoefficient[0][1] = FMath::Min(MinCoefficient[0][1], V);
+					MaxCoefficient[0][1] = FMath::Max(MaxCoefficient[0][1], V);
 
-				MinCoefficient[0][2] = FMath::Min(MinCoefficient[0][2], W);
-				MaxCoefficient[0][2] = FMath::Max(MaxCoefficient[0][2], W);
+					MinCoefficient[0][2] = FMath::Min(MinCoefficient[0][2], W);
+					MaxCoefficient[0][2] = FMath::Max(MaxCoefficient[0][2], W);
 
-				MinCoefficient[0][3] = FMath::Min(MinCoefficient[0][3], LogL);
-				MaxCoefficient[0][3] = FMath::Max(MaxCoefficient[0][3], LogL);
+					MinCoefficient[0][3] = FMath::Min(MinCoefficient[0][3], LogL);
+					MaxCoefficient[0][3] = FMath::Max(MaxCoefficient[0][3], LogL);
 
-				// Dampen dark texel's contribution on the directionality min and max
-				float DampenDirectionality = FMath::Clamp(L * 100, 0.0f, 1.0f);
+					// Dampen dark texel's contribution on the directionality min and max
+					float DampenDirectionality = FMath::Clamp(L * 100, 0.0f, 1.0f);
 
-				for (int32 ColorIndex = 0; ColorIndex < 3; ColorIndex++)
-				{
-					MinCoefficient[1][ColorIndex] = FMath::Min(MinCoefficient[1][ColorIndex], DampenDirectionality * SourceSample.Coefficients[1][ColorIndex]);
-					MaxCoefficient[1][ColorIndex] = FMath::Max(MaxCoefficient[1][ColorIndex], DampenDirectionality * SourceSample.Coefficients[1][ColorIndex]);
+					for (int32 ColorIndex = 0; ColorIndex < 3; ColorIndex++)
+					{
+						MinCoefficient[1][ColorIndex] = FMath::Min(MinCoefficient[1][ColorIndex], DampenDirectionality * SourceSample.Coefficients[1][ColorIndex]);
+						MaxCoefficient[1][ColorIndex] = FMath::Max(MaxCoefficient[1][ColorIndex], DampenDirectionality * SourceSample.Coefficients[1][ColorIndex]);
+					}
 				}
-			}
 
-			{
-				// Simple
-				float L, U, V, W;
-				GetLUVW(SourceSample.Coefficients[2], L, U, V, W);
+				{
+					// Simple
+					float L, U, V, W;
+					GetLUVW(SourceSample.Coefficients[2], L, U, V, W);
 
-				float LogL = FMath::Log2(L + SimpleLogBlackPoint) / SimpleLogScale + 0.5f;
+					float LogL = FMath::Log2(L + SimpleLogBlackPoint) / SimpleLogScale + 0.5f;
 
-				float LogR = LogL * U;
-				float LogG = LogL * V;
-				float LogB = LogL * W;
+					float LogR = LogL * U;
+					float LogG = LogL * V;
+					float LogB = LogL * W;
 
-				MinCoefficient[2][0] = FMath::Min(MinCoefficient[2][0], LogR);
-				MaxCoefficient[2][0] = FMath::Max(MaxCoefficient[2][0], LogR);
+					MinCoefficient[2][0] = FMath::Min(MinCoefficient[2][0], LogR);
+					MaxCoefficient[2][0] = FMath::Max(MaxCoefficient[2][0], LogR);
 
-				MinCoefficient[2][1] = FMath::Min(MinCoefficient[2][1], LogG);
-				MaxCoefficient[2][1] = FMath::Max(MaxCoefficient[2][1], LogG);
+					MinCoefficient[2][1] = FMath::Min(MinCoefficient[2][1], LogG);
+					MaxCoefficient[2][1] = FMath::Max(MaxCoefficient[2][1], LogG);
 
-				MinCoefficient[2][2] = FMath::Min(MinCoefficient[2][2], LogB);
-				MaxCoefficient[2][2] = FMath::Max(MaxCoefficient[2][2], LogB);
+					MinCoefficient[2][2] = FMath::Min(MinCoefficient[2][2], LogB);
+					MaxCoefficient[2][2] = FMath::Max(MaxCoefficient[2][2], LogB);
 
-				// Dampen dark texel's contribution on the directionality min and max
-				float DampenDirectionality = FMath::Clamp(L * 100, 0.0f, 1.0f);
+					// Dampen dark texel's contribution on the directionality min and max
+					float DampenDirectionality = FMath::Clamp(L * 100, 0.0f, 1.0f);
 
-				for (int32 ColorIndex = 0; ColorIndex < 3; ColorIndex++)
-				{
-					MinCoefficient[3][ColorIndex] = FMath::Min(MinCoefficient[3][ColorIndex], DampenDirectionality * SourceSample.Coefficients[3][ColorIndex]);
-					MaxCoefficient[3][ColorIndex] = FMath::Max(MaxCoefficient[3][ColorIndex], DampenDirectionality * SourceSample.Coefficients[3][ColorIndex]);
+					for (int32 ColorIndex = 0; ColorIndex < 3; ColorIndex++)
+					{
+						MinCoefficient[3][ColorIndex] = FMath::Min(MinCoefficient[3][ColorIndex], DampenDirectionality * SourceSample.Coefficients[3][ColorIndex]);
+						MaxCoefficient[3][ColorIndex] = FMath::Max(MaxCoefficient[3][ColorIndex], DampenDirectionality * SourceSample.Coefficients[3][ColorIndex]);
+					}
 				}
 			}
 		}
+	});
+	
+	float MinCoefficient[NUM_STORED_LIGHTMAP_COEF][4];
+	float MaxCoefficient[NUM_STORED_LIGHTMAP_COEF][4];
+	
+	for (int32 CoefficientIndex = 0; CoefficientIndex < NUM_STORED_LIGHTMAP_COEF; CoefficientIndex += 2)
+	{
+		for (int32 ColorIndex = 0; ColorIndex < 4; ColorIndex++)
+		{
+			// Color
+			MinCoefficient[CoefficientIndex][ColorIndex] = 10000.0f;
+			MaxCoefficient[CoefficientIndex][ColorIndex] = -10000.0f;
+
+			// Direction
+			MinCoefficient[CoefficientIndex + 1][ColorIndex] = 10000.0f;
+			MaxCoefficient[CoefficientIndex + 1][ColorIndex] = -10000.0f;
+		}
 	}
 
+	for(int32 TaskIndex = 0; TaskIndex < NumTasks; TaskIndex++)
+	{
+		for (int32 CoefficientIndex = 0; CoefficientIndex < NUM_STORED_LIGHTMAP_COEF; CoefficientIndex += 2)
+		{
+			for (int32 ColorIndex = 0; ColorIndex < 4; ColorIndex++)
+			{
+				// Color
+				MinCoefficient[CoefficientIndex][ColorIndex] = FMath::Min(MinCoefficient[CoefficientIndex][ColorIndex], TaskStorage[TaskIndex].MinCoefficient[CoefficientIndex][ColorIndex]);
+				MaxCoefficient[CoefficientIndex][ColorIndex] = FMath::Max(MaxCoefficient[CoefficientIndex][ColorIndex], TaskStorage[TaskIndex].MaxCoefficient[CoefficientIndex][ColorIndex]);
+	
+				// Direction
+				MinCoefficient[CoefficientIndex + 1][ColorIndex] = FMath::Min(MinCoefficient[CoefficientIndex + 1][ColorIndex],TaskStorage[TaskIndex].MinCoefficient[CoefficientIndex + 1][ColorIndex]);
+				MaxCoefficient[CoefficientIndex + 1][ColorIndex] = FMath::Max(MaxCoefficient[CoefficientIndex + 1][ColorIndex],TaskStorage[TaskIndex].MaxCoefficient[CoefficientIndex + 1][ColorIndex]);
+			}
+		}
+	}
+	
 	// If no sample mapped make range sane.
 	// Or if very dark no directionality is added. Make range sane.
 	for (int32 CoefficientIndex = 0; CoefficientIndex < NUM_STORED_LIGHTMAP_COEF; CoefficientIndex++)
@@ -172,88 +221,91 @@ void QuantizeLightSamples(
 	// allocate space in the output
 	OutLightSamples.Empty(InLightSamples.Num());
 	OutLightSamples.AddUninitialized(InLightSamples.Num());
-
-	// quantize each sample using the above scaling
-	for (int32 SampleIndex = 0; SampleIndex < InLightSamples.Num(); SampleIndex++)
+	
+	ParallelFor(NumTasks,[&](int32 TaskIndex)
 	{
-		const FLightSampleData& SourceSample = InLightSamples[SampleIndex];
-		FLightMapCoefficients& DestCoefficients = OutLightSamples[SampleIndex];
+		// quantize each sample using the above scaling
+		for (int32 SampleIndex = TaskIndex * NumSamplesPerTask; SampleIndex < (TaskIndex + 1) * NumSamplesPerTask && SampleIndex < InLightSamples.Num(); SampleIndex++)
+		{
+			const FLightSampleData& SourceSample = InLightSamples[SampleIndex];
+			FLightMapCoefficients& DestCoefficients = OutLightSamples[SampleIndex];
 
-		DestCoefficients.Coverage = SourceSample.bIsMapped ? 255 : 0;
+			DestCoefficients.Coverage = SourceSample.bIsMapped ? 255 : 0;
 
-		const FVector BentNormal(SourceSample.SkyOcclusion[0], SourceSample.SkyOcclusion[1], SourceSample.SkyOcclusion[2]);
-		const float BentNormalLength = BentNormal.Size();
-		const FVector NormalizedBentNormal = BentNormal.GetSafeNormal() * FVector(.5f) + FVector(.5f);
+			const FVector BentNormal(SourceSample.SkyOcclusion[0], SourceSample.SkyOcclusion[1], SourceSample.SkyOcclusion[2]);
+			const float BentNormalLength = BentNormal.Size();
+			const FVector NormalizedBentNormal = BentNormal.GetSafeNormal() * FVector(.5f) + FVector(.5f);
 
-		DestCoefficients.SkyOcclusion[0] = (uint8)FMath::Clamp<int32>(FMath::RoundToInt(NormalizedBentNormal[0] * 255.0f), 0, 255);
-		DestCoefficients.SkyOcclusion[1] = (uint8)FMath::Clamp<int32>(FMath::RoundToInt(NormalizedBentNormal[1] * 255.0f), 0, 255);
-		DestCoefficients.SkyOcclusion[2] = (uint8)FMath::Clamp<int32>(FMath::RoundToInt(NormalizedBentNormal[2] * 255.0f), 0, 255);
-		// Sqrt on length to allocate more precision near 0
-		DestCoefficients.SkyOcclusion[3] = (uint8)FMath::Clamp<int32>(FMath::RoundToInt(FMath::Sqrt(BentNormalLength) * 255.0f), 0, 255);
+			DestCoefficients.SkyOcclusion[0] = (uint8)FMath::Clamp<int32>(FMath::RoundToInt(NormalizedBentNormal[0] * 255.0f), 0, 255);
+			DestCoefficients.SkyOcclusion[1] = (uint8)FMath::Clamp<int32>(FMath::RoundToInt(NormalizedBentNormal[1] * 255.0f), 0, 255);
+			DestCoefficients.SkyOcclusion[2] = (uint8)FMath::Clamp<int32>(FMath::RoundToInt(NormalizedBentNormal[2] * 255.0f), 0, 255);
+			// Sqrt on length to allocate more precision near 0
+			DestCoefficients.SkyOcclusion[3] = (uint8)FMath::Clamp<int32>(FMath::RoundToInt(FMath::Sqrt(BentNormalLength) * 255.0f), 0, 255);
 
-		// Sqrt to allocate more precision near 0
-		DestCoefficients.AOMaterialMask = (uint8)FMath::Clamp<int32>(FMath::RoundToInt(FMath::Sqrt(SourceSample.AOMaterialMask) * 255.0f), 0, 255);
+			// Sqrt to allocate more precision near 0
+			DestCoefficients.AOMaterialMask = (uint8)FMath::Clamp<int32>(FMath::RoundToInt(FMath::Sqrt(SourceSample.AOMaterialMask) * 255.0f), 0, 255);
 
-		{
-			float L, U, V, W;
-			GetLUVW(SourceSample.Coefficients[0], L, U, V, W);
-
-			// LogLUVW encode color
-			float LogL = FMath::Log2(L + LogBlackPoint);
-
-			U = U * CoefficientMultiply[0][0] + CoefficientAdd[0][0];
-			V = V * CoefficientMultiply[0][1] + CoefficientAdd[0][1];
-			W = W * CoefficientMultiply[0][2] + CoefficientAdd[0][2];
-			LogL = LogL * CoefficientMultiply[0][3] + CoefficientAdd[0][3];
-
-			float Residual = LogL * 255.0f - FMath::RoundToFloat(LogL * 255.0f) + 0.5f;
-
-			// U, V, W, LogL
-			// UVW stored in gamma space
-			DestCoefficients.Coefficients[0][0] = (uint8)FMath::Clamp<int32>(FMath::RoundToInt(FMath::Sqrt(U) * 255.0f), 0, 255);
-			DestCoefficients.Coefficients[0][1] = (uint8)FMath::Clamp<int32>(FMath::RoundToInt(FMath::Sqrt(V) * 255.0f), 0, 255);
-			DestCoefficients.Coefficients[0][2] = (uint8)FMath::Clamp<int32>(FMath::RoundToInt(FMath::Sqrt(W) * 255.0f), 0, 255);
-			DestCoefficients.Coefficients[0][3] = (uint8)FMath::Clamp<int32>(FMath::RoundToInt(LogL * 255.0f), 0, 255);
-
-			float Dx = SourceSample.Coefficients[1][0] * CoefficientMultiply[1][0] + CoefficientAdd[1][0];
-			float Dy = SourceSample.Coefficients[1][1] * CoefficientMultiply[1][1] + CoefficientAdd[1][1];
-			float Dz = SourceSample.Coefficients[1][2] * CoefficientMultiply[1][2] + CoefficientAdd[1][2];
-
-			// Dx, Dy, Dz, Residual
-			DestCoefficients.Coefficients[1][0] = (uint8)FMath::Clamp<int32>(FMath::RoundToInt(Dx * 255.0f), 0, 255);
-			DestCoefficients.Coefficients[1][1] = (uint8)FMath::Clamp<int32>(FMath::RoundToInt(Dy * 255.0f), 0, 255);
-			DestCoefficients.Coefficients[1][2] = (uint8)FMath::Clamp<int32>(FMath::RoundToInt(Dz * 255.0f), 0, 255);
-			DestCoefficients.Coefficients[1][3] = (uint8)FMath::Clamp<int32>(FMath::RoundToInt(Residual * 255.0f), 0, 255);
-		}
+			{
+				float L, U, V, W;
+				GetLUVW(SourceSample.Coefficients[0], L, U, V, W);
 
-		{
-			float L, U, V, W;
-			GetLUVW(SourceSample.Coefficients[2], L, U, V, W);
-
-			// LogRGB encode color
-			float LogL = FMath::Log2(L + SimpleLogBlackPoint) / SimpleLogScale + 0.5f;
-
-			float LogR = LogL * U * CoefficientMultiply[2][0] + CoefficientAdd[2][0];
-			float LogG = LogL * V * CoefficientMultiply[2][1] + CoefficientAdd[2][1];
-			float LogB = LogL * W * CoefficientMultiply[2][2] + CoefficientAdd[2][2];
-
-			// LogR, LogG, LogB
-			DestCoefficients.Coefficients[2][0] = (uint8)FMath::Clamp<int32>(FMath::RoundToInt(LogR * 255.0f), 0, 255);
-			DestCoefficients.Coefficients[2][1] = (uint8)FMath::Clamp<int32>(FMath::RoundToInt(LogG * 255.0f), 0, 255);
-			DestCoefficients.Coefficients[2][2] = (uint8)FMath::Clamp<int32>(FMath::RoundToInt(LogB * 255.0f), 0, 255);
-			DestCoefficients.Coefficients[2][3] = 255;
-
-			float Dx = SourceSample.Coefficients[3][0] * CoefficientMultiply[3][0] + CoefficientAdd[3][0];
-			float Dy = SourceSample.Coefficients[3][1] * CoefficientMultiply[3][1] + CoefficientAdd[3][1];
-			float Dz = SourceSample.Coefficients[3][2] * CoefficientMultiply[3][2] + CoefficientAdd[3][2];
-
-			// Dx, Dy, Dz
-			DestCoefficients.Coefficients[3][0] = (uint8)FMath::Clamp<int32>(FMath::RoundToInt(Dx * 255.0f), 0, 255);
-			DestCoefficients.Coefficients[3][1] = (uint8)FMath::Clamp<int32>(FMath::RoundToInt(Dy * 255.0f), 0, 255);
-			DestCoefficients.Coefficients[3][2] = (uint8)FMath::Clamp<int32>(FMath::RoundToInt(Dz * 255.0f), 0, 255);
-			DestCoefficients.Coefficients[3][3] = 255;
+				// LogLUVW encode color
+				float LogL = FMath::Log2(L + LogBlackPoint);
+
+				U = U * CoefficientMultiply[0][0] + CoefficientAdd[0][0];
+				V = V * CoefficientMultiply[0][1] + CoefficientAdd[0][1];
+				W = W * CoefficientMultiply[0][2] + CoefficientAdd[0][2];
+				LogL = LogL * CoefficientMultiply[0][3] + CoefficientAdd[0][3];
+
+				float Residual = LogL * 255.0f - FMath::RoundToFloat(LogL * 255.0f) + 0.5f;
+
+				// U, V, W, LogL
+				// UVW stored in gamma space
+				DestCoefficients.Coefficients[0][0] = (uint8)FMath::Clamp<int32>(FMath::RoundToInt(FMath::Sqrt(U) * 255.0f), 0, 255);
+				DestCoefficients.Coefficients[0][1] = (uint8)FMath::Clamp<int32>(FMath::RoundToInt(FMath::Sqrt(V) * 255.0f), 0, 255);
+				DestCoefficients.Coefficients[0][2] = (uint8)FMath::Clamp<int32>(FMath::RoundToInt(FMath::Sqrt(W) * 255.0f), 0, 255);
+				DestCoefficients.Coefficients[0][3] = (uint8)FMath::Clamp<int32>(FMath::RoundToInt(LogL * 255.0f), 0, 255);
+
+				float Dx = SourceSample.Coefficients[1][0] * CoefficientMultiply[1][0] + CoefficientAdd[1][0];
+				float Dy = SourceSample.Coefficients[1][1] * CoefficientMultiply[1][1] + CoefficientAdd[1][1];
+				float Dz = SourceSample.Coefficients[1][2] * CoefficientMultiply[1][2] + CoefficientAdd[1][2];
+
+				// Dx, Dy, Dz, Residual
+				DestCoefficients.Coefficients[1][0] = (uint8)FMath::Clamp<int32>(FMath::RoundToInt(Dx * 255.0f), 0, 255);
+				DestCoefficients.Coefficients[1][1] = (uint8)FMath::Clamp<int32>(FMath::RoundToInt(Dy * 255.0f), 0, 255);
+				DestCoefficients.Coefficients[1][2] = (uint8)FMath::Clamp<int32>(FMath::RoundToInt(Dz * 255.0f), 0, 255);
+				DestCoefficients.Coefficients[1][3] = (uint8)FMath::Clamp<int32>(FMath::RoundToInt(Residual * 255.0f), 0, 255);
+			}
+
+			{
+				float L, U, V, W;
+				GetLUVW(SourceSample.Coefficients[2], L, U, V, W);
+
+				// LogRGB encode color
+				float LogL = FMath::Log2(L + SimpleLogBlackPoint) / SimpleLogScale + 0.5f;
+
+				float LogR = LogL * U * CoefficientMultiply[2][0] + CoefficientAdd[2][0];
+				float LogG = LogL * V * CoefficientMultiply[2][1] + CoefficientAdd[2][1];
+				float LogB = LogL * W * CoefficientMultiply[2][2] + CoefficientAdd[2][2];
+
+				// LogR, LogG, LogB
+				DestCoefficients.Coefficients[2][0] = (uint8)FMath::Clamp<int32>(FMath::RoundToInt(LogR * 255.0f), 0, 255);
+				DestCoefficients.Coefficients[2][1] = (uint8)FMath::Clamp<int32>(FMath::RoundToInt(LogG * 255.0f), 0, 255);
+				DestCoefficients.Coefficients[2][2] = (uint8)FMath::Clamp<int32>(FMath::RoundToInt(LogB * 255.0f), 0, 255);
+				DestCoefficients.Coefficients[2][3] = 255;
+
+				float Dx = SourceSample.Coefficients[3][0] * CoefficientMultiply[3][0] + CoefficientAdd[3][0];
+				float Dy = SourceSample.Coefficients[3][1] * CoefficientMultiply[3][1] + CoefficientAdd[3][1];
+				float Dz = SourceSample.Coefficients[3][2] * CoefficientMultiply[3][2] + CoefficientAdd[3][2];
+
+				// Dx, Dy, Dz
+				DestCoefficients.Coefficients[3][0] = (uint8)FMath::Clamp<int32>(FMath::RoundToInt(Dx * 255.0f), 0, 255);
+				DestCoefficients.Coefficients[3][1] = (uint8)FMath::Clamp<int32>(FMath::RoundToInt(Dy * 255.0f), 0, 255);
+				DestCoefficients.Coefficients[3][2] = (uint8)FMath::Clamp<int32>(FMath::RoundToInt(Dz * 255.0f), 0, 255);
+				DestCoefficients.Coefficients[3][3] = 255;
+			}
 		}
-	}
+	});
 }
 
 FLightSampleData ConvertToLightSample(FLinearColor IncidentLighting, FLinearColor LuminanceSH)
diff --git a/Engine/Plugins/Experimental/GPULightmass/Source/GPULightmass/Private/LightmapRenderer.h b/Engine/Plugins/Experimental/GPULightmass/Source/GPULightmass/Private/LightmapRenderer.h
--- a/Engine/Plugins/Experimental/GPULightmass/Source/GPULightmass/Private/LightmapRenderer.h
+++ b/Engine/Plugins/Experimental/GPULightmass/Source/GPULightmass/Private/LightmapRenderer.h
@@ -66,7 +66,7 @@ struct FLightmapReadbackGroup
 		// Duplicate some metadata so the async thread doesn't need to access FLightmapReadbackGroup
 		FIntPoint SizeInTiles;
 		int32 RowPitchInPixels[3];
-		TArray<FLinearColor, FAnsiAllocator> Texture[3];
+		TArray<FLinearColor> Texture[3];
 
 		volatile int32 bDenoisingFinished = 0;
 	};
@@ -84,7 +84,7 @@ struct FLightmapTileDenoiseGroup
 
 	struct FTextureData
 	{
-		TArray<FLinearColor, FAnsiAllocator> Texture[2];
+		TArray<FLinearColor> Texture[2];
 
 		volatile int32 bDenoisingFinished = 0;
 	};
diff --git a/Engine/Plugins/Experimental/GPULightmass/Source/GPULightmass/Private/LightmapStorage.cpp b/Engine/Plugins/Experimental/GPULightmass/Source/GPULightmass/Private/LightmapStorage.cpp
--- a/Engine/Plugins/Experimental/GPULightmass/Source/GPULightmass/Private/LightmapStorage.cpp
+++ b/Engine/Plugins/Experimental/GPULightmass/Source/GPULightmass/Private/LightmapStorage.cpp
@@ -9,6 +9,7 @@
 #include "LightmapRenderer.h"
 #include "GPULightmassModule.h"
 #include "Async/ParallelFor.h"
+#include "Compression/OodleDataCompression.h"
 
 namespace GPULightmass
 {
@@ -17,14 +18,19 @@ TDoubleLinkedList<FTileDataLayer*> FTileDataLayer::AllUncompressedTiles;
 
 int64 FTileDataLayer::Compress(bool bParallelCompression)
 {
-	const int32 CompressMemoryBound = FCompression::CompressMemoryBound(NAME_LZ4, sizeof(FLinearColor) * GPreviewLightmapVirtualTileSize * GPreviewLightmapVirtualTileSize);
+	const int32 CompressMemoryBound = FOodleDataCompression::CompressedBufferSizeNeeded(sizeof(FLinearColor) * GPreviewLightmapVirtualTileSize * GPreviewLightmapVirtualTileSize);
 
 	if (Data.Num() > 0)
 	{
 		CompressedData.Empty();
 		CompressedData.AddUninitialized(CompressMemoryBound);
 		int32 CompressedSize = CompressMemoryBound;
-		check(FCompression::CompressMemory(NAME_LZ4, CompressedData.GetData(), CompressedSize, Data.GetData(), sizeof(FLinearColor) * GPreviewLightmapVirtualTileSize * GPreviewLightmapVirtualTileSize));
+
+		CompressedSize = FOodleDataCompression::Compress(
+			CompressedData.GetData(), CompressedSize,
+			Data.GetData(), sizeof(FLinearColor) * GPreviewLightmapVirtualTileSize * GPreviewLightmapVirtualTileSize,
+			FOodleDataCompression::ECompressor::Selkie, FOodleDataCompression::ECompressionLevel::VeryFast);
+		
 		CompressedData.SetNum(CompressedSize);
 		Data.Empty();
 
@@ -58,8 +64,10 @@ void FTileDataLayer::Decompress()
 	if (CompressedData.Num() != 0)
 	{
 		Data.Empty();
-		Data.AddUninitialized(GPreviewLightmapVirtualTileSize * GPreviewLightmapVirtualTileSize);
-		check(FCompression::UncompressMemory(NAME_LZ4, Data.GetData(), sizeof(FLinearColor) * GPreviewLightmapVirtualTileSize * GPreviewLightmapVirtualTileSize, CompressedData.GetData(), CompressedData.Num()));
+		Data.AddUninitialized(GPreviewLightmapVirtualTileSize * GPreviewLightmapVirtualTileSize);		
+		check(FOodleDataCompression::Decompress(
+			Data.GetData(), sizeof(FLinearColor) * GPreviewLightmapVirtualTileSize * GPreviewLightmapVirtualTileSize,
+			CompressedData.GetData(), CompressedData.Num()));
 		CompressedData.Empty();
 	}
 	else
@@ -75,7 +83,7 @@ void FTileDataLayer::Evict()
 
 	TArray<FTileDataLayer*> TileDataLayersToCompress;
 
-	while (AllUncompressedTiles.Num() > 16384)
+	while (AllUncompressedTiles.Num() > 65536) // 4 gig
 	{
 		FTileDataLayer* TileDataLayerToCompress = AllUncompressedTiles.GetTail()->GetValue();
 		AllUncompressedTiles.RemoveNode(AllUncompressedTiles.GetTail(), false);
diff --git a/Engine/Plugins/Experimental/GPULightmass/Source/GPULightmass/Private/LightmapStorage.h b/Engine/Plugins/Experimental/GPULightmass/Source/GPULightmass/Private/LightmapStorage.h
--- a/Engine/Plugins/Experimental/GPULightmass/Source/GPULightmass/Private/LightmapStorage.h
+++ b/Engine/Plugins/Experimental/GPULightmass/Source/GPULightmass/Private/LightmapStorage.h
@@ -8,15 +8,13 @@
 #include "UObject/GCObjectScopeGuard.h"
 #include "VT/LightmapVirtualTexture.h"
 #include "Engine/MapBuildDataRegistry.h"
-// TBB suffers from extreme fragmentation problem in editor
-#include "Core/Private/HAL/Allocators/AnsiAllocator.h"
 
 namespace GPULightmass
 {
 
 struct FTileDataLayer
 {
-	TArray<FLinearColor, FAnsiAllocator> Data;
+	TArray<FLinearColor> Data;
 	TArray<uint8> CompressedData;
 
 	TDoubleLinkedList<FTileDataLayer*>::TDoubleLinkedListNode Node;
diff --git a/Engine/Plugins/Experimental/GPULightmass/Source/GPULightmass/Private/Scene/Scene.cpp b/Engine/Plugins/Experimental/GPULightmass/Source/GPULightmass/Private/Scene/Scene.cpp
--- a/Engine/Plugins/Experimental/GPULightmass/Source/GPULightmass/Private/Scene/Scene.cpp
+++ b/Engine/Plugins/Experimental/GPULightmass/Source/GPULightmass/Private/Scene/Scene.cpp
@@ -1719,8 +1719,9 @@ void CopyRectTiled(
 	int32 TileBorderSize = 0
 )
 {
-	for (int32 Y = DstRect.Min.Y; Y < DstRect.Max.Y; Y++)
+	ParallelFor(DstRect.Max.Y - DstRect.Min.Y, [&](int32 dy)
 	{
+		int32 Y = DstRect.Min.Y + dy;
 		for (int32 X = DstRect.Min.X; X < DstRect.Max.X; X++)
 		{
 			FIntPoint SrcPosition = FIntPoint(X, Y) - DstRect.Min + SrcMin;
@@ -1735,7 +1736,7 @@ void CopyRectTiled(
 
 			Func(DstLinearIndex, SrcTilePosition, SrcLinearIndex);
 		}
-	}
+	});
 }
 
 void ReadbackVolumetricLightmapDataLayerFromGPU(FRHICommandListImmediate& RHICmdList, FVolumetricLightmapDataLayer& Layer, FIntVector Dimensions)
diff --git a/Engine/Source/Runtime/Engine/Private/LightMap.cpp b/Engine/Source/Runtime/Engine/Private/LightMap.cpp
--- a/Engine/Source/Runtime/Engine/Private/LightMap.cpp
+++ b/Engine/Source/Runtime/Engine/Private/LightMap.cpp
@@ -2413,9 +2413,9 @@ void FLightMap2D::EncodeTextures( UWorld* InWorld, ULevel* LightingScenario, boo
 
 		if (!bIncludeNonVirtualTextures)
 		{
-			// If we exclusively using VT lightmaps, don't need to worry about max size of a given texture sheet
-			// Also just make square textures, don't need 2x1 aspect ratio
-			PackedLightAndShadowMapTextureSizeX = 32 * 1024; // FTextureLayout uses uint16 for size in some places
+			// If we exclusively using VT lightmaps, just make square textures, don't need 2x1 aspect ratio
+			// YW: while it's possible to create huge VT lightmaps (like 32k), it hurts multithread lightmap encoding performance seriously.
+			// Packed lightmaps with 1k - 2k resolution should be enough to reduce drawcalls
 			PackedLightAndShadowMapTextureSizeY = PackedLightAndShadowMapTextureSizeX;
 		}
 
