GPULM: Pause operation when shader compilation is active

diff --git a/Engine/Plugins/Experimental/GPULightmass/Source/GPULightmass/Private/LightmapRenderer.cpp b/Engine/Plugins/Experimental/GPULightmass/Source/GPULightmass/Private/LightmapRenderer.cpp
--- a/Engine/Plugins/Experimental/GPULightmass/Source/GPULightmass/Private/LightmapRenderer.cpp
+++ b/Engine/Plugins/Experimental/GPULightmass/Source/GPULightmass/Private/LightmapRenderer.cpp
@@ -26,6 +26,7 @@
 #include "RayTracingInstanceBufferUtil.h"
 #include "ScreenPass.h"
 #include "RayTracingDynamicGeometryCollection.h"
+#include "ShaderCompiler.h"
 
 class FCopyConvergedLightmapTilesCS : public FGlobalShader
 {
@@ -1421,6 +1422,13 @@ void FLightmapRenderer::Finalize(FRDGBuilder& GraphBuilder)
 {
 	TRACE_CPUPROFILER_EVENT_SCOPE(FLightmapRenderer::Finalize);
 
+	bool bIsCompilingShaders = GShaderCompilingManager && GShaderCompilingManager->IsCompiling();
+
+	if (bIsCompilingShaders)
+	{
+		return;
+	}
+	
 	if (PendingTileRequests.Num() == 0)
 	{
 		return;
@@ -3657,6 +3665,7 @@ void FLightmapRenderer::BackgroundTick()
 
 void FLightmapRenderer::BumpRevision()
 {
+	FrameNumber = 0;
 	CurrentRevision++;
 
 	for (TArray<FLightmapTileRequest>& FrameRequests : TilesVisibleLastFewFrames)
@@ -3664,6 +3673,7 @@ void FLightmapRenderer::BumpRevision()
 		FrameRequests.Empty();
 	}	
 
+	PendingTileRequests.Empty();
 	RecordedTileRequests.Empty();
 
 	LightmapTilePoolGPU.UnmapAll();
diff --git a/Engine/Plugins/Experimental/GPULightmass/Source/GPULightmass/Private/Scene/Scene.cpp b/Engine/Plugins/Experimental/GPULightmass/Source/GPULightmass/Private/Scene/Scene.cpp
--- a/Engine/Plugins/Experimental/GPULightmass/Source/GPULightmass/Private/Scene/Scene.cpp
+++ b/Engine/Plugins/Experimental/GPULightmass/Source/GPULightmass/Private/Scene/Scene.cpp
@@ -23,6 +23,7 @@
 #include "Editor.h"
 #include "Components/HierarchicalInstancedStaticMeshComponent.h"
 #include "LightmapDenoising.h"
+#include "ShaderCompiler.h"
 #include "Misc/FileHelper.h"
 #include "Components/ReflectionCaptureComponent.h"
 
@@ -1521,7 +1522,10 @@ void FScene::RemoveGeometryInstanceFromComponent(ULandscapeComponent* InComponen
 				{
 					const int8 SubSectionIdx = SubX + SubY * LandscapeRenderState.NumSubsections;
 
-					LandscapeRenderState.SectionRayTracingStates[SubSectionIdx]->Geometry.ReleaseResource();
+					if (LandscapeRenderState.SectionRayTracingStates[SubSectionIdx].IsValid())
+					{
+						LandscapeRenderState.SectionRayTracingStates[SubSectionIdx]->Geometry.ReleaseResource();
+					}
 				}
 			}
 		}
@@ -1592,27 +1596,35 @@ void FScene::BackgroundTick()
 
 	if (Percentage < 100 || GPULightmass->Settings->Mode == EGPULightmassMode::BakeWhatYouSee)
 	{
-		if (bNeedsVoxelization)
+		bool bIsCompilingShaders = GShaderCompilingManager && GShaderCompilingManager->IsCompiling();
+
+		if (!bIsCompilingShaders)
 		{
-			GatherImportanceVolumes();
+			if (bNeedsVoxelization)
+			{
+				GatherImportanceVolumes();
+
+				ENQUEUE_RENDER_COMMAND(BackgroundTickRenderThread)([&RenderState = RenderState](FRHICommandListImmediate&) mutable {
+					RenderState.VolumetricLightmapRenderer->VoxelizeScene();
+					RenderState.VolumetricLightmapRenderer->FrameNumber = 0;
+					RenderState.VolumetricLightmapRenderer->SamplesTaken = 0;
+				});
+
+				bNeedsVoxelization = false;
+			}
 
 			ENQUEUE_RENDER_COMMAND(BackgroundTickRenderThread)([&RenderState = RenderState](FRHICommandListImmediate&) mutable {
-				RenderState.VolumetricLightmapRenderer->VoxelizeScene();
-				RenderState.VolumetricLightmapRenderer->FrameNumber = 0;
-				RenderState.VolumetricLightmapRenderer->SamplesTaken = 0;
+				RenderState.BackgroundTick();
 			});
-
-			bNeedsVoxelization = false;
 		}
-
-		ENQUEUE_RENDER_COMMAND(BackgroundTickRenderThread)([&RenderState = RenderState](FRHICommandListImmediate&) mutable {
-			RenderState.BackgroundTick();
-		});
 	}
 	else
 	{
-		ApplyFinishedLightmapsToWorld();
-		GPULightmass->World->GetSubsystem<UGPULightmassSubsystem>()->OnLightBuildEnded().Broadcast();
+		if (!bNeedsVoxelization)
+		{
+			ApplyFinishedLightmapsToWorld();
+			GPULightmass->World->GetSubsystem<UGPULightmassSubsystem>()->OnLightBuildEnded().Broadcast();
+		}
 	}
 }
 
