GPULM: Refactor FGPUTileDescription to remove duplicated code

diff --git a/Engine/Plugins/Experimental/GPULightmass/Shaders/Private/LightmapOutput.usf b/Engine/Plugins/Experimental/GPULightmass/Shaders/Private/LightmapOutput.usf
--- a/Engine/Plugins/Experimental/GPULightmass/Shaders/Private/LightmapOutput.usf
+++ b/Engine/Plugins/Experimental/GPULightmass/Shaders/Private/LightmapOutput.usf
@@ -94,4 +94,7 @@ void SelectiveLightmapOutputCS(uint3 DispatchThreadId : SV_DispatchThreadID)
 
 		OutputTileAtlas[LaunchIndex + BatchedTiles[TileIndex].OutputLayer2Position] = sqrt((float4)ShadowValue / max(uint4(1, 1, 1, 1), SampleCountValue)) * ValidityMask;
 	}
+	else if (DIM_OUTPUT_LAYER == 3)
+	{
+	}
 }
diff --git a/Engine/Plugins/Experimental/GPULightmass/Source/GPULightmass/Private/LightmapRenderer.cpp b/Engine/Plugins/Experimental/GPULightmass/Source/GPULightmass/Private/LightmapRenderer.cpp
--- a/Engine/Plugins/Experimental/GPULightmass/Source/GPULightmass/Private/LightmapRenderer.cpp
+++ b/Engine/Plugins/Experimental/GPULightmass/Source/GPULightmass/Private/LightmapRenderer.cpp
@@ -89,7 +89,7 @@ class FSelectiveLightmapOutputCS : public FGlobalShader
 	DECLARE_GLOBAL_SHADER(FSelectiveLightmapOutputCS)
 	SHADER_USE_PARAMETER_STRUCT(FSelectiveLightmapOutputCS, FGlobalShader)
 
-	class FOutputLayerDim : SHADER_PERMUTATION_INT("DIM_OUTPUT_LAYER", 3);
+	class FOutputLayerDim : SHADER_PERMUTATION_INT("DIM_OUTPUT_LAYER", 4);
 	class FDrawProgressBars : SHADER_PERMUTATION_BOOL("DRAW_PROGRESS_BARS");
 	using FPermutationDomain = TShaderPermutationDomain<FOutputLayerDim, FDrawProgressBars>;
 
@@ -161,15 +161,46 @@ struct FGPUTileDescription
 	int32 RenderPassIndex;
 };
 
+namespace GPULightmass
+{
+
 struct FGPUBatchedTileRequests
 {
 	FBufferRHIRef BatchedTilesBuffer;
 	FShaderResourceViewRHIRef BatchedTilesSRV;
 	TResourceArray<FGPUTileDescription> BatchedTilesDesc;
-};
 
-namespace GPULightmass
-{
+	void BuildFromTileDescs(TArray<FLightmapTileRequest>& TileRequests, FLightmapTilePoolGPU& LightmapTilePoolGPU, FLightmapTilePoolGPU& ScratchTilePoolGPU)
+	{
+		for (const FLightmapTileRequest& Tile : TileRequests)
+		{
+			FGPUTileDescription TileDesc;
+			TileDesc.LightmapSize = Tile.RenderState->GetSize();
+			TileDesc.VirtualTilePosition = Tile.VirtualCoordinates.Position * GPreviewLightmapVirtualTileSize;
+			TileDesc.WorkingSetPosition = LightmapTilePoolGPU.GetPositionFromLinearAddress(Tile.TileAddressInWorkingSet) * GPreviewLightmapPhysicalTileSize;
+			TileDesc.ScratchPosition = ScratchTilePoolGPU.GetPositionFromLinearAddress(Tile.TileAddressInScratch) * GPreviewLightmapPhysicalTileSize;
+			TileDesc.OutputLayer0Position = Tile.OutputPhysicalCoordinates[0] * GPreviewLightmapPhysicalTileSize;
+			TileDesc.OutputLayer1Position = Tile.OutputPhysicalCoordinates[1] * GPreviewLightmapPhysicalTileSize;
+			TileDesc.OutputLayer2Position = Tile.OutputPhysicalCoordinates[2] * GPreviewLightmapPhysicalTileSize;
+			TileDesc.FrameIndex = Tile.RenderState->RetrieveTileState(Tile.VirtualCoordinates).Revision;
+			TileDesc.RenderPassIndex = Tile.RenderState->RetrieveTileState(Tile.VirtualCoordinates).RenderPassIndex;
+			BatchedTilesDesc.Add(TileDesc);
+		}
+	}
+
+	void Commit()
+	{
+		if (BatchedTilesDesc.Num() > 0)
+		{
+			FRHIResourceCreateInfo CreateInfo(TEXT("BatchedTilesBuffer"));
+			CreateInfo.GPUMask = FRHIGPUMask::GPU0();
+			CreateInfo.ResourceArray = &BatchedTilesDesc;
+
+			BatchedTilesBuffer = RHICreateStructuredBuffer(sizeof(FGPUTileDescription), BatchedTilesDesc.GetResourceDataSize(), BUF_Dynamic | BUF_ShaderResource, CreateInfo);
+			BatchedTilesSRV = RHICreateShaderResourceView(BatchedTilesBuffer);
+		}
+	}
+};
 
 FLightmapRenderer::FLightmapRenderer(FSceneRenderState* InScene)
 	: Scene(InScene)
@@ -1572,30 +1603,8 @@ void FLightmapRenderer::Finalize(FRDGBuilder& GraphBuilder)
 			}
 
 			FGPUBatchedTileRequests GPUBatchedTileRequests;
-
-			{
-				for (const FLightmapTileRequest& Tile : TileUploadRequests)
-				{
-					FGPUTileDescription TileDesc;
-					TileDesc.LightmapSize = Tile.RenderState->GetSize();
-					TileDesc.VirtualTilePosition = Tile.VirtualCoordinates.Position * GPreviewLightmapVirtualTileSize;
-					TileDesc.WorkingSetPosition = LightmapTilePoolGPU.GetPositionFromLinearAddress(Tile.TileAddressInWorkingSet) * GPreviewLightmapPhysicalTileSize;
-					TileDesc.ScratchPosition = ScratchTilePoolGPU->GetPositionFromLinearAddress(Tile.TileAddressInScratch) * GPreviewLightmapPhysicalTileSize;
-					TileDesc.OutputLayer0Position = Tile.OutputPhysicalCoordinates[0] * GPreviewLightmapPhysicalTileSize;
-					TileDesc.OutputLayer1Position = Tile.OutputPhysicalCoordinates[1] * GPreviewLightmapPhysicalTileSize;
-					TileDesc.OutputLayer2Position = Tile.OutputPhysicalCoordinates[2] * GPreviewLightmapPhysicalTileSize;
-					TileDesc.FrameIndex = Tile.RenderState->RetrieveTileState(Tile.VirtualCoordinates).Revision;
-					TileDesc.RenderPassIndex = Tile.RenderState->RetrieveTileState(Tile.VirtualCoordinates).RenderPassIndex;
-					GPUBatchedTileRequests.BatchedTilesDesc.Add(TileDesc);
-				}
-
-				FRHIResourceCreateInfo CreateInfo(TEXT("BatchedTilesBuffer"));
-				CreateInfo.GPUMask = FRHIGPUMask::GPU0();
-				CreateInfo.ResourceArray = &GPUBatchedTileRequests.BatchedTilesDesc;
-
-				GPUBatchedTileRequests.BatchedTilesBuffer = RHICreateStructuredBuffer(sizeof(FGPUTileDescription), GPUBatchedTileRequests.BatchedTilesDesc.GetResourceDataSize(), BUF_Dynamic | BUF_ShaderResource, CreateInfo);
-				GPUBatchedTileRequests.BatchedTilesSRV = RHICreateShaderResourceView(GPUBatchedTileRequests.BatchedTilesBuffer);
-			}
+			GPUBatchedTileRequests.BuildFromTileDescs(TileUploadRequests, LightmapTilePoolGPU, *ScratchTilePoolGPU);
+			GPUBatchedTileRequests.Commit();
 
 			IPooledRenderTarget* OutputRenderTargets[3] = { nullptr, nullptr, nullptr };
 
@@ -2039,11 +2048,12 @@ void FLightmapRenderer::Finalize(FRDGBuilder& GraphBuilder)
 
 	RDG_GPU_MASK_SCOPE(GraphBuilder, FRHIGPUMask::GPU0());
 
-	IPooledRenderTarget* OutputRenderTargets[3] = { nullptr, nullptr, nullptr };
-
+	IPooledRenderTarget* OutputRenderTargets[5] = { nullptr, nullptr, nullptr, nullptr, nullptr };
+	static_assert(UE_ARRAY_COUNT(OutputRenderTargets) <= UE_ARRAY_COUNT(FLightmapTileRequest::OutputRenderTargets));
+	
 	for (auto& Tile : PendingTileRequests)
 	{
-		for (int32 RenderTargetIndex = 0; RenderTargetIndex < 3; RenderTargetIndex++)
+		for (int32 RenderTargetIndex = 0; RenderTargetIndex < UE_ARRAY_COUNT(OutputRenderTargets); RenderTargetIndex++)
 		{
 			if (Tile.OutputRenderTargets[RenderTargetIndex] != nullptr)
 			{
@@ -2289,21 +2299,15 @@ void FLightmapRenderer::Finalize(FRDGBuilder& GraphBuilder)
 					{
 						FGPUBatchedTileRequests GPUBatchedTileRequests;
 
+						TArray<FLightmapTileRequest> TileRequestsThisGPU;
+						
 						for (const FLightmapTileRequest& Tile : PendingGITileRequests)
 						{
 							uint32 AssignedGPUIndex = (Tile.RenderState->DistributionPrefixSum + Tile.RenderState->RetrieveTileStateIndex(Tile.VirtualCoordinates)) % GNumExplicitGPUsForRendering;
 							if (AssignedGPUIndex != GPUIndex) continue;
 
-							FGPUTileDescription TileDesc;
-							TileDesc.LightmapSize = Tile.RenderState->GetSize();
-							TileDesc.VirtualTilePosition = Tile.VirtualCoordinates.Position * GPreviewLightmapVirtualTileSize;
-							TileDesc.WorkingSetPosition = LightmapTilePoolGPU.GetPositionFromLinearAddress(Tile.TileAddressInWorkingSet) * GPreviewLightmapPhysicalTileSize;
-							TileDesc.ScratchPosition = ScratchTilePoolGPU->GetPositionFromLinearAddress(Tile.TileAddressInScratch) * GPreviewLightmapPhysicalTileSize;
-							TileDesc.OutputLayer0Position = Tile.OutputPhysicalCoordinates[0] * GPreviewLightmapPhysicalTileSize;
-							TileDesc.OutputLayer1Position = Tile.OutputPhysicalCoordinates[1] * GPreviewLightmapPhysicalTileSize;
-							TileDesc.OutputLayer2Position = Tile.OutputPhysicalCoordinates[2] * GPreviewLightmapPhysicalTileSize;
-							TileDesc.FrameIndex = Tile.RenderState->RetrieveTileState(Tile.VirtualCoordinates).Revision;
-							TileDesc.RenderPassIndex = Tile.RenderState->RetrieveTileState(Tile.VirtualCoordinates).RenderPassIndex;
+							TileRequestsThisGPU.Add(Tile);
+							
 							if (!Tile.RenderState->IsTileGIConverged(Tile.VirtualCoordinates, NumTotalPassesToRender))
 							{
 								Tile.RenderState->RetrieveTileState(Tile.VirtualCoordinates).RenderPassIndex++;
@@ -2315,24 +2319,17 @@ void FLightmapRenderer::Finalize(FRDGBuilder& GraphBuilder)
 										Mip0WorkDoneLastFrame++;
 									}
 								}
-
-								GPUBatchedTileRequests.BatchedTilesDesc.Add(TileDesc);
 							}
 						}
 
-						if (GPUBatchedTileRequests.BatchedTilesDesc.Num() > 0)
-						{
-							FRHIResourceCreateInfo CreateInfo(TEXT("BatchedTilesBuffer"));
-							CreateInfo.GPUMask = FRHIGPUMask::FromIndex(GPUIndex);
-							CreateInfo.ResourceArray = &GPUBatchedTileRequests.BatchedTilesDesc;
-
-							GPUBatchedTileRequests.BatchedTilesBuffer = RHICreateStructuredBuffer(sizeof(FGPUTileDescription), GPUBatchedTileRequests.BatchedTilesDesc.GetResourceDataSize(), BUF_Dynamic | BUF_ShaderResource, CreateInfo);
-							GPUBatchedTileRequests.BatchedTilesSRV = HoldReference(GraphBuilder, RHICreateShaderResourceView(GPUBatchedTileRequests.BatchedTilesBuffer));
-						}
+						GPUBatchedTileRequests.BuildFromTileDescs(TileRequestsThisGPU, LightmapTilePoolGPU, *ScratchTilePoolGPU);
+						GPUBatchedTileRequests.Commit();
+						// Let GraphBuilder references the SRV
+						GraphBuilder.AllocObject<FShaderResourceViewRHIRef>(GPUBatchedTileRequests.BatchedTilesSRV);
 
 						RDG_GPU_MASK_SCOPE(GraphBuilder, FRHIGPUMask::FromIndex(GPUIndex));
 
-						if (GPUBatchedTileRequests.BatchedTilesDesc.Num() > 0)
+						if (TileRequestsThisGPU.Num() > 0)
 						{
 							FRDGTextureRef GBufferWorldPosition = GraphBuilder.RegisterExternalTexture(ScratchTilePoolGPU->PooledRenderTargets[0], TEXT("GBufferWorldPosition"));
 							FRDGTextureRef GBufferWorldNormal = GraphBuilder.RegisterExternalTexture(ScratchTilePoolGPU->PooledRenderTargets[1], TEXT("GBufferWorldNormal"));
@@ -2732,34 +2729,16 @@ void FLightmapRenderer::Finalize(FRDGBuilder& GraphBuilder)
 #if RHI_RAYTRACING
 					if (IsRayTracingEnabled())
 					{
-						FGPUBatchedTileRequests GPUBatchedTileRequests;
-
+						FGPUBatchedTileRequests GPUBatchedTileRequests;						
+						GPUBatchedTileRequests.BuildFromTileDescs(PendingShadowTileRequests,LightmapTilePoolGPU, *ScratchTilePoolGPU);
+						for (int32 TileIndex = 0; TileIndex < PendingShadowTileRequests.Num(); TileIndex++)
 						{
-							for (int32 TileIndex = 0; TileIndex < PendingShadowTileRequests.Num(); TileIndex++)
-							{
-								const FLightmapTileRequest& Tile = PendingShadowTileRequests[TileIndex];
-								FGPUTileDescription TileDesc;
-								TileDesc.LightmapSize = Tile.RenderState->GetSize();
-								TileDesc.VirtualTilePosition = Tile.VirtualCoordinates.Position * GPreviewLightmapVirtualTileSize;
-								TileDesc.WorkingSetPosition = LightmapTilePoolGPU.GetPositionFromLinearAddress(Tile.TileAddressInWorkingSet) * GPreviewLightmapPhysicalTileSize;
-								TileDesc.ScratchPosition = ScratchTilePoolGPU->GetPositionFromLinearAddress(Tile.TileAddressInScratch) * GPreviewLightmapPhysicalTileSize;
-								TileDesc.OutputLayer0Position = Tile.OutputPhysicalCoordinates[0] * GPreviewLightmapPhysicalTileSize;
-								TileDesc.OutputLayer1Position = Tile.OutputPhysicalCoordinates[1] * GPreviewLightmapPhysicalTileSize;
-								TileDesc.OutputLayer2Position = Tile.OutputPhysicalCoordinates[2] * GPreviewLightmapPhysicalTileSize;
-								TileDesc.FrameIndex = Tile.RenderState->RetrieveTileState(Tile.VirtualCoordinates).Revision;
-								TileDesc.RenderPassIndex = LightSampleIndexArray[TileIndex];
-								GPUBatchedTileRequests.BatchedTilesDesc.Add(TileDesc);
-							}
-
-							{
-								FRHIResourceCreateInfo CreateInfo(TEXT("BatchedTilesBuffer"));
-								CreateInfo.GPUMask = FRHIGPUMask::FromIndex(GPUIndex);
-								CreateInfo.ResourceArray = &GPUBatchedTileRequests.BatchedTilesDesc;
-
-								GPUBatchedTileRequests.BatchedTilesBuffer = RHICreateStructuredBuffer(sizeof(FGPUTileDescription), GPUBatchedTileRequests.BatchedTilesDesc.GetResourceDataSize(), BUF_Dynamic | BUF_ShaderResource, CreateInfo);
-								GPUBatchedTileRequests.BatchedTilesSRV = HoldReference(GraphBuilder, RHICreateShaderResourceView(GPUBatchedTileRequests.BatchedTilesBuffer));
-							}
+							GPUBatchedTileRequests.BatchedTilesDesc[TileIndex].RenderPassIndex = LightSampleIndexArray[TileIndex];
 						}
+						GPUBatchedTileRequests.Commit();
+						
+						// Let GraphBuilder references the SRV
+						GraphBuilder.AllocObject<FShaderResourceViewRHIRef>(GPUBatchedTileRequests.BatchedTilesSRV);
 
 						FIntPoint RayTracingResolution;
 						RayTracingResolution.X = GPreviewLightmapPhysicalTileSize * GPUBatchedTileRequests.BatchedTilesDesc.Num();
@@ -2841,30 +2820,10 @@ void FLightmapRenderer::Finalize(FRDGBuilder& GraphBuilder)
 	// Output from working set to VT layers
 	{
 		FGPUBatchedTileRequests GPUBatchedTileRequests;
-
-		{
-			for (const FLightmapTileRequest& Tile : PendingTileRequests)
-			{
-				FGPUTileDescription TileDesc;
-				TileDesc.LightmapSize = Tile.RenderState->GetSize();
-				TileDesc.VirtualTilePosition = Tile.VirtualCoordinates.Position * GPreviewLightmapVirtualTileSize;
-				TileDesc.WorkingSetPosition = LightmapTilePoolGPU.GetPositionFromLinearAddress(Tile.TileAddressInWorkingSet) * GPreviewLightmapPhysicalTileSize;
-				TileDesc.ScratchPosition = ScratchTilePoolGPU->GetPositionFromLinearAddress(Tile.TileAddressInScratch) * GPreviewLightmapPhysicalTileSize;
-				TileDesc.OutputLayer0Position = Tile.OutputPhysicalCoordinates[0] * GPreviewLightmapPhysicalTileSize;
-				TileDesc.OutputLayer1Position = Tile.OutputPhysicalCoordinates[1] * GPreviewLightmapPhysicalTileSize;
-				TileDesc.OutputLayer2Position = Tile.OutputPhysicalCoordinates[2] * GPreviewLightmapPhysicalTileSize;
-				TileDesc.FrameIndex = Tile.RenderState->RetrieveTileState(Tile.VirtualCoordinates).Revision;
-				TileDesc.RenderPassIndex = Tile.RenderState->RetrieveTileState(Tile.VirtualCoordinates).RenderPassIndex;
-				GPUBatchedTileRequests.BatchedTilesDesc.Add(TileDesc);
-			}
-
-			FRHIResourceCreateInfo CreateInfo(TEXT("BatchedTilesBuffer"));
-			CreateInfo.GPUMask = FRHIGPUMask::GPU0();
-			CreateInfo.ResourceArray = &GPUBatchedTileRequests.BatchedTilesDesc;
-
-			GPUBatchedTileRequests.BatchedTilesBuffer = RHICreateStructuredBuffer(sizeof(FGPUTileDescription), GPUBatchedTileRequests.BatchedTilesDesc.GetResourceDataSize(), BUF_Dynamic | BUF_ShaderResource, CreateInfo);
-			GPUBatchedTileRequests.BatchedTilesSRV = HoldReference(GraphBuilder, RHICreateShaderResourceView(GPUBatchedTileRequests.BatchedTilesBuffer));
-		}
+		GPUBatchedTileRequests.BuildFromTileDescs(PendingTileRequests, LightmapTilePoolGPU, *ScratchTilePoolGPU);
+		GPUBatchedTileRequests.Commit();
+		// Let GraphBuilder references the SRV
+		GraphBuilder.AllocObject<FShaderResourceViewRHIRef>(GPUBatchedTileRequests.BatchedTilesSRV);
 
 		{
 			FRDGTextureRef IrradianceAndSampleCount = GraphBuilder.RegisterExternalTexture(LightmapTilePoolGPU.PooledRenderTargets[0], TEXT("IrradianceAndSampleCount"));
@@ -2912,7 +2871,7 @@ void FLightmapRenderer::Finalize(FRDGBuilder& GraphBuilder)
 
 				FSelectiveLightmapOutputCS::FPermutationDomain PermutationVector;
 				PermutationVector.Set<FSelectiveLightmapOutputCS::FOutputLayerDim>(2);
-				PermutationVector.Set<FSelectiveLightmapOutputCS::FDrawProgressBars>(Scene->Settings->bShowProgressBars);
+				PermutationVector.Set<FSelectiveLightmapOutputCS::FDrawProgressBars>(false);
 
 				auto Shader = GlobalShaderMap->GetShader<FSelectiveLightmapOutputCS>(PermutationVector);
 
@@ -2931,6 +2890,31 @@ void FLightmapRenderer::Finalize(FRDGBuilder& GraphBuilder)
 					PassParameters,
 					FComputeShaderUtils::GetGroupCount(RayTracingResolution, FComputeShaderUtils::kGolden2DGroupSize));
 			}
+			
+			if (OutputRenderTargets[3] != nullptr)
+			{
+				FRDGTextureRef RenderTargetTileAtlas = GraphBuilder.RegisterExternalTexture(OutputRenderTargets[3], TEXT("GPULightmassRenderTargetTileAtlas3"));
+
+				FSelectiveLightmapOutputCS::FPermutationDomain PermutationVector;
+				PermutationVector.Set<FSelectiveLightmapOutputCS::FOutputLayerDim>(3);
+				PermutationVector.Set<FSelectiveLightmapOutputCS::FDrawProgressBars>(false);
+
+				auto Shader = GlobalShaderMap->GetShader<FSelectiveLightmapOutputCS>(PermutationVector);
+
+				FSelectiveLightmapOutputCS::FParameters* PassParameters = GraphBuilder.AllocParameters<FSelectiveLightmapOutputCS::FParameters>();
+				PassParameters->NumBatchedTiles = GPUBatchedTileRequests.BatchedTilesDesc.Num();
+				PassParameters->BatchedTiles = GPUBatchedTileRequests.BatchedTilesSRV;
+				PassParameters->OutputTileAtlas = GraphBuilder.CreateUAV(RenderTargetTileAtlas);
+				PassParameters->IrradianceAndSampleCount = GraphBuilder.CreateUAV(IrradianceAndSampleCount);
+				PassParameters->SHCorrectionAndStationarySkyLightBentNormal = GraphBuilder.CreateUAV(SHCorrectionAndStationarySkyLightBentNormal);
+
+				FComputeShaderUtils::AddPass(
+					GraphBuilder,
+					RDG_EVENT_NAME("SelectiveLightmapOutput 3"),
+					Shader,
+					PassParameters,
+					FComputeShaderUtils::GetGroupCount(RayTracingResolution, FComputeShaderUtils::kGolden2DGroupSize));
+			}
 		}
 	}
 
@@ -3008,32 +2992,10 @@ void FLightmapRenderer::Finalize(FRDGBuilder& GraphBuilder)
 				}
 
 				FGPUBatchedTileRequests GPUBatchedTileRequests;
-
-				for (const auto& Tile : LightmapReadbackGroup.ConvergedTileRequests)
-				{
-					uint32 AssignedGPUIndex = (Tile.RenderState->DistributionPrefixSum + Tile.RenderState->RetrieveTileStateIndex(Tile.VirtualCoordinates)) % GNumExplicitGPUsForRendering;
-					check(AssignedGPUIndex == GPUIndex);
-
-					FGPUTileDescription TileDesc;
-					TileDesc.LightmapSize = Tile.RenderState->GetSize();
-					TileDesc.VirtualTilePosition = Tile.VirtualCoordinates.Position * GPreviewLightmapVirtualTileSize;
-					TileDesc.WorkingSetPosition = LightmapTilePoolGPU.GetPositionFromLinearAddress(Tile.TileAddressInWorkingSet) * GPreviewLightmapPhysicalTileSize;
-					TileDesc.ScratchPosition = ScratchTilePoolGPU->GetPositionFromLinearAddress(Tile.TileAddressInScratch) * GPreviewLightmapPhysicalTileSize;
-					TileDesc.OutputLayer0Position = Tile.OutputPhysicalCoordinates[0] * GPreviewLightmapPhysicalTileSize;
-					TileDesc.OutputLayer1Position = Tile.OutputPhysicalCoordinates[1] * GPreviewLightmapPhysicalTileSize;
-					TileDesc.OutputLayer2Position = Tile.OutputPhysicalCoordinates[2] * GPreviewLightmapPhysicalTileSize;
-					TileDesc.FrameIndex = Tile.RenderState->RetrieveTileState(Tile.VirtualCoordinates).Revision;
-					TileDesc.RenderPassIndex = Tile.RenderState->RetrieveTileState(Tile.VirtualCoordinates).RenderPassIndex;
-					GPUBatchedTileRequests.BatchedTilesDesc.Add(TileDesc);
-				}
-
-				FRHIResourceCreateInfo CreateInfo(TEXT("BatchedTilesBuffer"));
-				CreateInfo.GPUMask = FRHIGPUMask::FromIndex(GPUIndex);
-				CreateInfo.ResourceArray = &GPUBatchedTileRequests.BatchedTilesDesc;
-
-				GPUBatchedTileRequests.BatchedTilesBuffer = RHICreateStructuredBuffer(sizeof(FGPUTileDescription), GPUBatchedTileRequests.BatchedTilesDesc.GetResourceDataSize(), BUF_Dynamic | BUF_ShaderResource, CreateInfo);
-				GPUBatchedTileRequests.BatchedTilesSRV = HoldReference(GraphBuilder, RHICreateShaderResourceView(GPUBatchedTileRequests.BatchedTilesBuffer));
-
+				GPUBatchedTileRequests.BuildFromTileDescs(LightmapReadbackGroup.ConvergedTileRequests, LightmapTilePoolGPU, *ScratchTilePoolGPU);
+				GPUBatchedTileRequests.Commit();
+				// Let GraphBuilder references the SRV
+				GraphBuilder.AllocObject<FShaderResourceViewRHIRef>(GPUBatchedTileRequests.BatchedTilesSRV);
 
 				FIntPoint DispatchResolution;
 				DispatchResolution.X = GPreviewLightmapPhysicalTileSize * GPUBatchedTileRequests.BatchedTilesDesc.Num();
