GPULM: Fix stationary light direct lighting not being ignored

diff --git a/Engine/Plugins/Experimental/GPULightmass/Shaders/Private/LightmapPathTracing.usf b/Engine/Plugins/Experimental/GPULightmass/Shaders/Private/LightmapPathTracing.usf
--- a/Engine/Plugins/Experimental/GPULightmass/Shaders/Private/LightmapPathTracing.usf
+++ b/Engine/Plugins/Experimental/GPULightmass/Shaders/Private/LightmapPathTracing.usf
@@ -316,8 +316,8 @@ void PathTracingKernel(
 	inout bool bIsValidSample,
 	inout float3 RadianceValue,
 	inout float3 RadianceDirection,
-	inout float3 DirectLightingNEERadianceValue,
-	inout float3 DirectLightingNEERadianceDirection,
+	inout float3 DirectLightingRadianceValue,
+	inout float3 DirectLightingRadianceDirection,
 	inout float4 PrimaryRandSample)
 {
 	uint2 LaunchIndex = uint2(0, 0);
@@ -587,9 +587,11 @@ void PathTracingKernel(
 
 						// Record the contribution
 						float3 LightContrib = PathThroughput * LightSample.RadianceOverPdf * MaterialEval.Weight * MaterialEval.Pdf;
+						float3 MISWeightedDirection = LightSample.Direction;
 						if (MISMode == 2)
 						{
 							LightContrib *= MISWeightRobust(LightSample.Pdf, MaterialEval.Pdf);
+							MISWeightedDirection *= MISWeightRobust(LightSample.Pdf, MaterialEval.Pdf);
 						}
 
 						// Record the contribution
@@ -599,11 +601,11 @@ void PathTracingKernel(
 						}
 						else
 						{
-							// GPULightmass records contribution of direct lighting separately, and only for static lights
+							// GPU Lightmass records contribution of direct lighting separately, and only for static lights
 							if (!IsStationary(LightId))
 							{
-								DirectLightingNEERadianceValue += LightContrib;
-								DirectLightingNEERadianceDirection = LightSample.Direction;
+								DirectLightingRadianceValue += LightContrib;
+								DirectLightingRadianceDirection += MISWeightedDirection;
 							}
 						}
 					}
@@ -705,7 +707,15 @@ void PathTracingKernel(
 					continue;
 				}
 
+				// GPU Lightmass records contribution of direct lighting separately, and only for static lights
+				if (Bounce == 0 && IsStationary(LightId))
+				{
+					continue;
+				}
+
 				float3 LightContrib = PathThroughput * LightResult.Radiance;
+				float3 MISWeightedDirection = MaterialSample.Direction;
+				
 				if (bUseMIS)
 				{
 					float PreviousCdfValue = 0.0;
@@ -716,6 +726,9 @@ void PathTracingKernel(
 					float LightPickPdf = (LightPickingCdf[Index] - PreviousCdfValue) / LightPickingCdfSum;
 
 					LightContrib *= MISWeightRobust(MaterialSample.Pdf, LightResult.Pdf * LightPickPdf);
+
+					// Separate direct lighting (bounce 0) for GPU Lightmass
+					MISWeightedDirection *= MISWeightRobust(MaterialSample.Pdf, LightResult.Pdf * LightPickPdf);
 				}
 
 				if (any(LightContrib > 0))
@@ -724,10 +737,21 @@ void PathTracingKernel(
 					{
 						RayDesc LightRay = Ray;
 						LightRay.TMax = LightResult.HitT;
-						LightContrib *= TraceTransparentVisibilityRay(LightRay, LaunchIndex, PathRoughness);
+						float3 Visibility = TraceTransparentVisibilityRay(LightRay, LaunchIndex, PathRoughness);
+						LightContrib *= Visibility;
+						
+						// Separate direct lighting (bounce 0) for GPU Lightmass
+						MISWeightedDirection *= Visibility;
 					}
 					// the light made some contribution, and there was nothing along the shadow ray
 					AccumulateRadiance(Radiance, LightContrib);
+
+					// Separate direct lighting (bounce 0) for GPU Lightmass
+					if (Bounce == 0)
+					{
+						DirectLightingRadianceValue += LightContrib;
+						DirectLightingRadianceDirection += MISWeightedDirection;
+					}
 				}
 			}
 			
@@ -823,8 +847,8 @@ void LightmapPathTracingMainRG()
 
 				float3 RadianceValue = 0;
 				float3 RadianceDirection = 0;				
-				float3 DirectLightingNEERadianceValue = 0;
-				float3 DirectLightingNEERadianceDirection = 0;
+				float3 DirectLightingRadianceValue = 0;
+				float3 DirectLightingRadianceDirection = 0;
 				
 				if (bIsValidSample)
 				{
@@ -837,8 +861,8 @@ void LightmapPathTracingMainRG()
 						bIsValidSample,
 						RadianceValue,
 						RadianceDirection,
-						DirectLightingNEERadianceValue,
-						DirectLightingNEERadianceDirection,
+						DirectLightingRadianceValue,
+						DirectLightingRadianceDirection,
 						PrimaryRandSample
 					);
 				}
@@ -886,10 +910,10 @@ void LightmapPathTracingMainRG()
 					{
 						FL2SHAndCorrection SH = (FL2SHAndCorrection)0;
 						
-						float TangentZ = saturate(dot(DirectLightingNEERadianceDirection, ShadingNormal));
-						SH.AddIncomingRadiance(Luminance(DirectLightingNEERadianceValue), DirectLightingNEERadianceDirection, TangentZ);
+						float TangentZ = saturate(dot(DirectLightingRadianceDirection, ShadingNormal));
+						SH.AddIncomingRadiance(Luminance(DirectLightingRadianceValue), DirectLightingRadianceDirection, TangentZ);
 						
-						IrradianceAndSampleCount[TexelIndexInPool].rgb += float3(DirectLightingNEERadianceValue * TangentZ / PI);
+						IrradianceAndSampleCount[TexelIndexInPool].rgb += float3(DirectLightingRadianceValue * TangentZ / PI);
 						SHDirectionality[TexelIndexInPool] += SH.L2SHCoefficients;					
 						SHCorrectionAndStationarySkyLightBentNormal[TexelIndexInPool].x += SH.Correction;
 					}
@@ -1140,8 +1164,8 @@ void VolumetricLightmapPathTracingMainRG()
 
 	float3 RadianceValue = 0;
 	float3 RadianceDirection = 0;				
-	float3 DirectLightingNEERadianceValue = 0;
-	float3 DirectLightingNEERadianceDirection = 0;
+	float3 DirectLightingRadianceValue = 0;
+	float3 DirectLightingRadianceDirection = 0;
 	float4 PrimaryRandSample;
 
 	if (bIsValidSample)
@@ -1159,15 +1183,15 @@ void VolumetricLightmapPathTracingMainRG()
 			bIsValidSample,
 			RadianceValue,
 			RadianceDirection,
-			DirectLightingNEERadianceValue,
-			DirectLightingNEERadianceDirection,
+			DirectLightingRadianceValue,
+			DirectLightingRadianceDirection,
 			PrimaryRandSample);
 	}
 	
 	// Every even sample is upper hemisphere, odd sample is lower hemisphere
 	// SampleCount stored is twice as the actual full sphere sample count
 	RadianceValue *= 2;
-	DirectLightingNEERadianceValue *= 2;
+	DirectLightingRadianceValue *= 2;
 
 	if (any(isnan(RadianceValue)) || any(RadianceValue < 0))
 	{
@@ -1223,15 +1247,15 @@ void VolumetricLightmapPathTracingMainRG()
 		
 		{
 			FThreeBandSHVectorRGBFloat SH;
-			SH.R.V0 = SHBasisFunction3Float(DirectLightingNEERadianceDirection).V0 * DirectLightingNEERadianceValue.r;
-			SH.R.V1 = SHBasisFunction3Float(DirectLightingNEERadianceDirection).V1 * DirectLightingNEERadianceValue.r;
-			SH.R.V2 = SHBasisFunction3Float(DirectLightingNEERadianceDirection).V2 * DirectLightingNEERadianceValue.r;
-			SH.G.V0 = SHBasisFunction3Float(DirectLightingNEERadianceDirection).V0 * DirectLightingNEERadianceValue.g;
-			SH.G.V1 = SHBasisFunction3Float(DirectLightingNEERadianceDirection).V1 * DirectLightingNEERadianceValue.g;
-			SH.G.V2 = SHBasisFunction3Float(DirectLightingNEERadianceDirection).V2 * DirectLightingNEERadianceValue.g;
-			SH.B.V0 = SHBasisFunction3Float(DirectLightingNEERadianceDirection).V0 * DirectLightingNEERadianceValue.b;
-			SH.B.V1 = SHBasisFunction3Float(DirectLightingNEERadianceDirection).V1 * DirectLightingNEERadianceValue.b;
-			SH.B.V2 = SHBasisFunction3Float(DirectLightingNEERadianceDirection).V2 * DirectLightingNEERadianceValue.b;
+			SH.R.V0 = SHBasisFunction3Float(DirectLightingRadianceDirection).V0 * DirectLightingRadianceValue.r;
+			SH.R.V1 = SHBasisFunction3Float(DirectLightingRadianceDirection).V1 * DirectLightingRadianceValue.r;
+			SH.R.V2 = SHBasisFunction3Float(DirectLightingRadianceDirection).V2 * DirectLightingRadianceValue.r;
+			SH.G.V0 = SHBasisFunction3Float(DirectLightingRadianceDirection).V0 * DirectLightingRadianceValue.g;
+			SH.G.V1 = SHBasisFunction3Float(DirectLightingRadianceDirection).V1 * DirectLightingRadianceValue.g;
+			SH.G.V2 = SHBasisFunction3Float(DirectLightingRadianceDirection).V2 * DirectLightingRadianceValue.g;
+			SH.B.V0 = SHBasisFunction3Float(DirectLightingRadianceDirection).V0 * DirectLightingRadianceValue.b;
+			SH.B.V1 = SHBasisFunction3Float(DirectLightingRadianceDirection).V1 * DirectLightingRadianceValue.b;
+			SH.B.V2 = SHBasisFunction3Float(DirectLightingRadianceDirection).V2 * DirectLightingRadianceValue.b;
 			
 			AmbientVector[VoxelPos].x += SH.R.V0.x;
 			AmbientVector[VoxelPos].y += SH.G.V0.x;
