GPULM: Implement support for lights' cast shadow flag, and proper support for stationary light shadow masks

diff --git a/Engine/Plugins/Experimental/GPULightmass/Source/GPULightmass/Private/PathTracingLightParameters.inl b/Engine/Plugins/Experimental/GPULightmass/Source/GPULightmass/Private/PathTracingLightParameters.inl
--- a/Engine/Plugins/Experimental/GPULightmass/Source/GPULightmass/Private/PathTracingLightParameters.inl
+++ b/Engine/Plugins/Experimental/GPULightmass/Source/GPULightmass/Private/PathTracingLightParameters.inl
@@ -26,7 +26,7 @@ void SetupPathTracingLightParameters(
 		DestLight.Color = FVector3f(LightScene.SkyLight->Color);
 		DestLight.Flags = PATHTRACER_FLAG_TRANSMISSION_MASK;
 		DestLight.Flags |= PATHTRACER_FLAG_LIGHTING_CHANNEL_MASK;
-		DestLight.Flags |= PATHTRACER_FLAG_CAST_SHADOW_MASK;
+		DestLight.Flags |= LightScene.SkyLight->bCastShadow ? PATHTRACER_FLAG_CAST_SHADOW_MASK : 0;
 		bool SkyLightIsStationary = LightScene.SkyLight->bStationary;
 		DestLight.Flags |= SkyLightIsStationary ? PATHTRACER_FLAG_STATIONARY_MASK : 0;
 		DestLight.Flags |= PATHTRACING_LIGHT_SKY;
@@ -61,7 +61,7 @@ void SetupPathTracingLightParameters(
 
 		DestLight.Flags = PATHTRACER_FLAG_TRANSMISSION_MASK;
 		DestLight.Flags |= PATHTRACER_FLAG_LIGHTING_CHANNEL_MASK;
-		DestLight.Flags |= PATHTRACER_FLAG_CAST_SHADOW_MASK;
+		DestLight.Flags |= Light.bCastShadow ? PATHTRACER_FLAG_CAST_SHADOW_MASK : 0;
 		DestLight.Flags |= Light.bStationary ? PATHTRACER_FLAG_STATIONARY_MASK : 0;
 		DestLight.Flags |= PATHTRACING_LIGHT_DIRECTIONAL;
 	}
@@ -95,7 +95,7 @@ void SetupPathTracingLightParameters(
 
 		DestLight.Flags = PATHTRACER_FLAG_TRANSMISSION_MASK;
 		DestLight.Flags |= PATHTRACER_FLAG_LIGHTING_CHANNEL_MASK;
-		DestLight.Flags |= PATHTRACER_FLAG_CAST_SHADOW_MASK;
+		DestLight.Flags |= Light.bCastShadow ? PATHTRACER_FLAG_CAST_SHADOW_MASK : 0;
 		DestLight.Flags |= Light.IsInverseSquared ? 0 : PATHTRACER_FLAG_NON_INVERSE_SQUARE_FALLOFF_MASK;
 		DestLight.Flags |= Light.bStationary ? PATHTRACER_FLAG_STATIONARY_MASK : 0;
 		DestLight.Flags |= PATHTRACING_LIGHT_POINT;
@@ -133,7 +133,7 @@ void SetupPathTracingLightParameters(
 
 		DestLight.Flags = PATHTRACER_FLAG_TRANSMISSION_MASK;
 		DestLight.Flags |= PATHTRACER_FLAG_LIGHTING_CHANNEL_MASK;
-		DestLight.Flags |= PATHTRACER_FLAG_CAST_SHADOW_MASK;
+		DestLight.Flags |= Light.bCastShadow ? PATHTRACER_FLAG_CAST_SHADOW_MASK : 0;
 		DestLight.Flags |= Light.IsInverseSquared ? 0 : PATHTRACER_FLAG_NON_INVERSE_SQUARE_FALLOFF_MASK;
 		DestLight.Flags |= Light.bStationary ? PATHTRACER_FLAG_STATIONARY_MASK : 0;
 		DestLight.Flags |= PATHTRACING_LIGHT_SPOT;
@@ -188,7 +188,7 @@ void SetupPathTracingLightParameters(
 
 		DestLight.Flags = PATHTRACER_FLAG_TRANSMISSION_MASK;
 		DestLight.Flags |= PATHTRACER_FLAG_LIGHTING_CHANNEL_MASK;
-		DestLight.Flags |= PATHTRACER_FLAG_CAST_SHADOW_MASK;
+		DestLight.Flags |= Light.bCastShadow ? PATHTRACER_FLAG_CAST_SHADOW_MASK : 0;
 		DestLight.Flags |= Light.bStationary ? PATHTRACER_FLAG_STATIONARY_MASK : 0;
 		DestLight.Flags |= PATHTRACING_LIGHT_RECT;
 
diff --git a/Engine/Plugins/Experimental/GPULightmass/Source/GPULightmass/Private/Scene/Lights.cpp b/Engine/Plugins/Experimental/GPULightmass/Source/GPULightmass/Private/Scene/Lights.cpp
--- a/Engine/Plugins/Experimental/GPULightmass/Source/GPULightmass/Private/Scene/Lights.cpp
+++ b/Engine/Plugins/Experimental/GPULightmass/Source/GPULightmass/Private/Scene/Lights.cpp
@@ -38,38 +38,41 @@ FLocalLightRenderState& FLightRenderStateRef::Resolve()
 	return LightRenderStateArray.ResolveAsLocalLightRenderState(*this);
 }
 
-FDirectionalLightBuildInfo::FDirectionalLightBuildInfo(UDirectionalLightComponent* DirectionalLightComponent)
+FLocalLightBuildInfo::FLocalLightBuildInfo(ULightComponent* LightComponent)
+{
+	bStationary = LightComponent->Mobility == EComponentMobility::Stationary;
+	bCastShadow = LightComponent->CastShadows && LightComponent->CastStaticShadows;
+	ShadowMapChannel = LightComponent->PreviewShadowMapChannel;
+	LightComponentMapBuildData = MakeUnique<FLightComponentMapBuildData>();
+	LightComponentMapBuildData->ShadowMapChannel = ShadowMapChannel;
+}
+
+FLocalLightRenderState::FLocalLightRenderState(ULightComponent* LightComponent)
 {
-	const bool bCastStationaryShadows = DirectionalLightComponent->CastShadows && DirectionalLightComponent->CastStaticShadows && !DirectionalLightComponent->HasStaticLighting();
+	bStationary = LightComponent->Mobility == EComponentMobility::Stationary;
+	bCastShadow = LightComponent->CastShadows && LightComponent->CastStaticShadows;
+	ShadowMapChannel = LightComponent->PreviewShadowMapChannel;
+}
 
+FDirectionalLightBuildInfo::FDirectionalLightBuildInfo(UDirectionalLightComponent* DirectionalLightComponent)
+	: FLocalLightBuildInfo(DirectionalLightComponent)
+{
 	ComponentUObject = DirectionalLightComponent;
-	bStationary = bCastStationaryShadows;
-	ShadowMapChannel = DirectionalLightComponent->PreviewShadowMapChannel;
-	LightComponentMapBuildData = MakeUnique<FLightComponentMapBuildData>();
-	LightComponentMapBuildData->ShadowMapChannel = DirectionalLightComponent->PreviewShadowMapChannel;
 }
 
 FDirectionalLightRenderState::FDirectionalLightRenderState(UDirectionalLightComponent* DirectionalLightComponent)
+	: FLocalLightRenderState(DirectionalLightComponent)
 {
-	const bool bCastStationaryShadows = DirectionalLightComponent->CastShadows && DirectionalLightComponent->CastStaticShadows && !DirectionalLightComponent->HasStaticLighting();
-
-	bStationary = bCastStationaryShadows;
 	Color = DirectionalLightComponent->GetColoredLightBrightness();
 	Direction = DirectionalLightComponent->GetDirection();
 	LightSourceAngle = DirectionalLightComponent->LightSourceAngle;
 	LightSourceSoftAngle = DirectionalLightComponent->LightSourceSoftAngle;
-	ShadowMapChannel = DirectionalLightComponent->PreviewShadowMapChannel;
 }
 
 FPointLightBuildInfo::FPointLightBuildInfo(UPointLightComponent* PointLightComponent)
+	: FLocalLightBuildInfo(PointLightComponent)
 {
-	const bool bCastStationaryShadows = PointLightComponent->CastShadows && PointLightComponent->CastStaticShadows && !PointLightComponent->HasStaticLighting();
-
 	ComponentUObject = PointLightComponent;
-	bStationary = bCastStationaryShadows;
-	ShadowMapChannel = PointLightComponent->PreviewShadowMapChannel;
-	LightComponentMapBuildData = MakeUnique<FLightComponentMapBuildData>();
-	LightComponentMapBuildData->ShadowMapChannel = PointLightComponent->PreviewShadowMapChannel;
 	Position = PointLightComponent->GetLightPosition();
 	AttenuationRadius = PointLightComponent->AttenuationRadius;
 }
@@ -80,10 +83,8 @@ bool FPointLightBuildInfo::AffectsBounds(const FBoxSphereBounds& InBounds) const
 }
 
 FPointLightRenderState::FPointLightRenderState(UPointLightComponent* PointLightComponent)
+	: FLocalLightRenderState(PointLightComponent)
 {
-	const bool bCastStationaryShadows = PointLightComponent->CastShadows && PointLightComponent->CastStaticShadows && !PointLightComponent->HasStaticLighting();
-
-	bStationary = bCastStationaryShadows;
 	Color = PointLightComponent->GetColoredLightBrightness();
 	Position = PointLightComponent->GetLightPosition();
 	Direction = PointLightComponent->GetDirection();
@@ -102,14 +103,9 @@ FPointLightRenderState::FPointLightRenderState(UPointLightComponent* PointLightC
 }
 
 FSpotLightBuildInfo::FSpotLightBuildInfo(USpotLightComponent* SpotLightComponent)
+	: FLocalLightBuildInfo(SpotLightComponent)
 {
-	const bool bCastStationaryShadows = SpotLightComponent->CastShadows && SpotLightComponent->CastStaticShadows && !SpotLightComponent->HasStaticLighting();
-
 	ComponentUObject = SpotLightComponent;
-	bStationary = bCastStationaryShadows;
-	ShadowMapChannel = SpotLightComponent->PreviewShadowMapChannel;
-	LightComponentMapBuildData = MakeUnique<FLightComponentMapBuildData>();
-	LightComponentMapBuildData->ShadowMapChannel = SpotLightComponent->PreviewShadowMapChannel;
 	Position = SpotLightComponent->GetLightPosition();
 	Direction = SpotLightComponent->GetDirection();
 	AttenuationRadius = SpotLightComponent->AttenuationRadius;
@@ -147,10 +143,8 @@ bool FSpotLightBuildInfo::AffectsBounds(const FBoxSphereBounds& InBounds) const
 }
 
 FSpotLightRenderState::FSpotLightRenderState(USpotLightComponent* SpotLightComponent)
+	: FLocalLightRenderState(SpotLightComponent)
 {
-	const bool bCastStationaryShadows = SpotLightComponent->CastShadows && SpotLightComponent->CastStaticShadows && !SpotLightComponent->HasStaticLighting();
-
-	bStationary = bCastStationaryShadows;
 	Color = SpotLightComponent->GetColoredLightBrightness();
 	Position = SpotLightComponent->GetLightPosition();
 	Direction = SpotLightComponent->GetDirection();
@@ -177,14 +171,9 @@ FSpotLightRenderState::FSpotLightRenderState(USpotLightComponent* SpotLightCompo
 }
 
 FRectLightBuildInfo::FRectLightBuildInfo(URectLightComponent* RectLightComponent)
+	: FLocalLightBuildInfo(RectLightComponent)
 {
-	const bool bCastStationaryShadows = RectLightComponent->CastShadows && RectLightComponent->CastStaticShadows && !RectLightComponent->HasStaticLighting();
-
 	ComponentUObject = RectLightComponent;
-	bStationary = bCastStationaryShadows;
-	ShadowMapChannel = RectLightComponent->PreviewShadowMapChannel;
-	LightComponentMapBuildData = MakeUnique<FLightComponentMapBuildData>();
-	LightComponentMapBuildData->ShadowMapChannel = RectLightComponent->PreviewShadowMapChannel;
 	Position = RectLightComponent->GetLightPosition();
 	AttenuationRadius = RectLightComponent->AttenuationRadius;
 }
@@ -195,10 +184,8 @@ bool FRectLightBuildInfo::AffectsBounds(const FBoxSphereBounds& InBounds) const
 }
 
 FRectLightRenderState::FRectLightRenderState(URectLightComponent* RectLightComponent)
+	: FLocalLightRenderState(RectLightComponent)
 {
-	const bool bCastStationaryShadows = RectLightComponent->CastShadows && RectLightComponent->CastStaticShadows && !RectLightComponent->HasStaticLighting();
-
-	bStationary = bCastStationaryShadows;
 	Color = RectLightComponent->GetColoredLightBrightness();
 	Position = RectLightComponent->GetLightPosition();
 	Direction = RectLightComponent->GetDirection();
diff --git a/Engine/Plugins/Experimental/GPULightmass/Source/GPULightmass/Private/Scene/Lights.h b/Engine/Plugins/Experimental/GPULightmass/Source/GPULightmass/Private/Scene/Lights.h
--- a/Engine/Plugins/Experimental/GPULightmass/Source/GPULightmass/Private/Scene/Lights.h
+++ b/Engine/Plugins/Experimental/GPULightmass/Source/GPULightmass/Private/Scene/Lights.h
@@ -65,24 +65,30 @@ struct FLightSceneRenderState;
 
 struct FLocalLightBuildInfo
 {
-	FLocalLightBuildInfo() = default;
+	FLocalLightBuildInfo(ULightComponent* LightComponent);
 	virtual ~FLocalLightBuildInfo() {}
 	FLocalLightBuildInfo(FLocalLightBuildInfo&& In) = default; // Default move constructor is killed by the custom destructor
 
 	bool bStationary = false;
+	bool bCastShadow = true;
 	int ShadowMapChannel = INDEX_NONE;
 
 	TUniquePtr<FLightComponentMapBuildData> LightComponentMapBuildData;
 
 	virtual bool AffectsBounds(const FBoxSphereBounds& InBounds) const = 0;
 	virtual ULightComponent* GetComponentUObject() const = 0;
+	bool CastsStationaryShadow() { return bStationary && bCastShadow; }
+
+	void AllocateMapBuildData(ULevel* StorageLevel);
 };
 
 struct FLocalLightRenderState
 {
+	FLocalLightRenderState(ULightComponent* LightComponent);
 	virtual ~FLocalLightRenderState() {}
 
 	bool bStationary = false;
+	bool bCastShadow = true;
 	int ShadowMapChannel = INDEX_NONE;
 
 	virtual FLightRenderParameters GetLightShaderParameters() const = 0;
@@ -281,6 +287,7 @@ using FRectLightRenderStateRef = TEntityArray<FRectLightRenderState>::EntityRefT
 struct FSkyLightRenderState
 {
 	bool bStationary = false;
+	bool bCastShadow = true;
 	FLinearColor Color;
 	FTextureRHIRef ProcessedTexture;
 	FSamplerStateRHIRef ProcessedTextureSampler;
diff --git a/Engine/Plugins/Experimental/GPULightmass/Source/GPULightmass/Private/Scene/Scene.cpp b/Engine/Plugins/Experimental/GPULightmass/Source/GPULightmass/Private/Scene/Scene.cpp
--- a/Engine/Plugins/Experimental/GPULightmass/Source/GPULightmass/Private/Scene/Scene.cpp
+++ b/Engine/Plugins/Experimental/GPULightmass/Source/GPULightmass/Private/Scene/Scene.cpp
@@ -224,7 +224,7 @@ void AddLightToLightmap(
 	// For both static and stationary lights
 	Lightmap.LightmapObject->LightGuids.Add(Light.GetComponentUObject()->LightGuid);
 
-	if (Light.bStationary)
+	if (Light.CastsStationaryShadow())
 	{
 		Lightmap.NumStationaryLightsPerShadowChannel[Light.ShadowMapChannel]++;
 		Lightmap.LightmapObject->bShadowChannelValid[Light.ShadowMapChannel] = true;
@@ -241,7 +241,7 @@ void RemoveLightFromLightmap(
 {
 	Lightmap.LightmapObject->LightGuids.Remove(Light.GetComponentUObject()->LightGuid);
 
-	if (Light.bStationary)
+	if (Light.CastsStationaryShadow())
 	{
 		Lightmap.NumStationaryLightsPerShadowChannel[Light.ShadowMapChannel]--;
 
@@ -366,31 +366,31 @@ struct LightTypeInfo<URectLightComponent>
 };
 
 template<typename LightComponentType>
-void FScene::AddLight(LightComponentType* PointLightComponent)
+void FScene::AddLight(LightComponentType* LightComponent)
 {
-	if (LightTypeInfo<LightComponentType>::GetLightComponentRegistration(LightScene).Contains(PointLightComponent))
+	if (LightTypeInfo<LightComponentType>::GetLightComponentRegistration(LightScene).Contains(LightComponent))
 	{
 		UE_LOG(LogGPULightmass, Log, TEXT("Warning: duplicated component registration"));
 		return;
 	}
 
-	const bool bCastStationaryShadows = PointLightComponent->CastShadows && PointLightComponent->CastStaticShadows && !PointLightComponent->HasStaticLighting();
+	const bool bCastStationaryShadows = LightComponent->CastShadows && LightComponent->CastStaticShadows && LightComponent->Mobility == EComponentMobility::Stationary;
 
 	if (bCastStationaryShadows)
 	{
-		if (PointLightComponent->PreviewShadowMapChannel == INDEX_NONE)
+		if (LightComponent->PreviewShadowMapChannel == INDEX_NONE)
 		{
 			UE_LOG(LogGPULightmass, Log, TEXT("Ignoring light with ShadowMapChannel == -1 (probably in the middle of SpawnActor)"));
 			return;
 		}
 	}
 
-	typename LightTypeInfo<LightComponentType>::BuildInfoType Light(PointLightComponent);
+	typename LightTypeInfo<LightComponentType>::BuildInfoType Light(LightComponent);
 
 	typename LightTypeInfo<LightComponentType>::LightRefType LightRef = LightTypeInfo<LightComponentType>::GetLightArray(LightScene).Emplace(MoveTemp(Light));
-	LightTypeInfo<LightComponentType>::GetLightComponentRegistration(LightScene).Add(PointLightComponent, LightRef);
+	LightTypeInfo<LightComponentType>::GetLightComponentRegistration(LightScene).Add(LightComponent, LightRef);
 
-	typename LightTypeInfo<LightComponentType>::RenderStateType LightRenderState(PointLightComponent);
+	typename LightTypeInfo<LightComponentType>::RenderStateType LightRenderState(LightComponent);
 
 	TArray<FPrimitiveSceneProxy*> SceneProxiesToUpdateOnRenderThread;
 	TArray<FGeometryRenderStateToken> RelevantGeometriesToUpdateOnRenderThread;
@@ -401,7 +401,7 @@ void FScene::AddLight(LightComponentType* PointLightComponent)
 
 		if (Light.AffectsBounds(Geometry.WorldBounds))
 		{
-			if (Light.bStationary)
+			if (Light.CastsStationaryShadow())
 			{
 				RelevantGeometriesToUpdateOnRenderThread.Add({ GeomIt.Index, GeomIt.Array.GetRenderStateArray() });
 			}
@@ -599,7 +599,8 @@ void FScene::AddLight(USkyLightComponent* SkyLight)
 	LightScene.SkyLight = MoveTemp(NewSkyLight);
 
 	FSkyLightRenderState NewSkyLightRenderState;
-	NewSkyLightRenderState.bStationary = !SkyLight->HasStaticLighting();
+	NewSkyLightRenderState.bStationary = SkyLight->CastShadows && SkyLight->CastStaticShadows && !SkyLight->HasStaticLighting();
+	NewSkyLightRenderState.bCastShadow = SkyLight->CastShadows && SkyLight->CastStaticShadows;
 	NewSkyLightRenderState.Color = SkyLight->GetLightColor() * SkyLight->Intensity;
 	NewSkyLightRenderState.TextureDimensions = FIntPoint(SkyLight->GetProcessedSkyTexture()->GetSizeX(), SkyLight->GetProcessedSkyTexture()->GetSizeY());
 	NewSkyLightRenderState.IrradianceEnvironmentMap = SkyLight->GetIrradianceEnvironmentMap();
@@ -658,7 +659,7 @@ TArray<int32> AddAllPossiblyRelevantLightsToGeometry(
 	{
 		if (Light.AffectsBounds(Instance->WorldBounds))
 		{
-			if (Light.bStationary)
+			if (Light.CastsStationaryShadow())
 			{
 				RelevantLightsToAddOnRenderThread.Add(&Light - LightArray.Elements.GetData());
 			}
@@ -1841,6 +1842,15 @@ void GatherBuildDataResourcesToKeep(const ULevel* InLevel, ULevel* LightingScena
 	}
 }
 
+void FLocalLightBuildInfo::AllocateMapBuildData(ULevel* StorageLevel)
+{
+	check(!CastsStationaryShadow() || GetComponentUObject()->PreviewShadowMapChannel != INDEX_NONE);
+
+	UMapBuildDataRegistry* Registry = StorageLevel->GetOrCreateMapBuildData();
+	FLightComponentMapBuildData& LightBuildData = Registry->FindOrAllocateLightBuildData(GetComponentUObject()->LightGuid, true);
+	LightBuildData.ShadowMapChannel = CastsStationaryShadow() ? GetComponentUObject()->PreviewShadowMapChannel : INDEX_NONE;
+}
+
 void FScene::ApplyFinishedLightmapsToWorld()
 {
 	UWorld* World = GPULightmass->World;
@@ -1898,46 +1908,22 @@ void FScene::ApplyFinishedLightmapsToWorld()
 
 		for (FDirectionalLightBuildInfo& DirectionalLight : LightScene.DirectionalLights.Elements)
 		{
-			UDirectionalLightComponent* Light = DirectionalLight.ComponentUObject;
-			check(!DirectionalLight.bStationary || Light->PreviewShadowMapChannel != INDEX_NONE);
-
-			ULevel* StorageLevel = LightingScenario ? LightingScenario : Light->GetOwner()->GetLevel();
-			UMapBuildDataRegistry* Registry = StorageLevel->GetOrCreateMapBuildData();
-			FLightComponentMapBuildData& LightBuildData = Registry->FindOrAllocateLightBuildData(Light->LightGuid, true);
-			LightBuildData.ShadowMapChannel = DirectionalLight.bStationary ? Light->PreviewShadowMapChannel : INDEX_NONE;
+			DirectionalLight.AllocateMapBuildData(LightingScenario ? LightingScenario : DirectionalLight.ComponentUObject->GetOwner()->GetLevel());
 		}
 
 		for (FPointLightBuildInfo& PointLight : LightScene.PointLights.Elements)
 		{
-			UPointLightComponent* Light = PointLight.ComponentUObject;
-			check(!PointLight.bStationary || Light->PreviewShadowMapChannel != INDEX_NONE);
-
-			ULevel* StorageLevel = LightingScenario ? LightingScenario : Light->GetOwner()->GetLevel();
-			UMapBuildDataRegistry* Registry = StorageLevel->GetOrCreateMapBuildData();
-			FLightComponentMapBuildData& LightBuildData = Registry->FindOrAllocateLightBuildData(Light->LightGuid, true);
-			LightBuildData.ShadowMapChannel = PointLight.bStationary ? Light->PreviewShadowMapChannel : INDEX_NONE;
+			PointLight.AllocateMapBuildData(LightingScenario ? LightingScenario : PointLight.ComponentUObject->GetOwner()->GetLevel());
 		}
 
 		for (FSpotLightBuildInfo& SpotLight : LightScene.SpotLights.Elements)
 		{
-			USpotLightComponent* Light = SpotLight.ComponentUObject;
-			check(!SpotLight.bStationary || Light->PreviewShadowMapChannel != INDEX_NONE);
-
-			ULevel* StorageLevel = LightingScenario ? LightingScenario : Light->GetOwner()->GetLevel();
-			UMapBuildDataRegistry* Registry = StorageLevel->GetOrCreateMapBuildData();
-			FLightComponentMapBuildData& LightBuildData = Registry->FindOrAllocateLightBuildData(Light->LightGuid, true);
-			LightBuildData.ShadowMapChannel = SpotLight.bStationary ? Light->PreviewShadowMapChannel : INDEX_NONE;
+			SpotLight.AllocateMapBuildData(LightingScenario ? LightingScenario : SpotLight.ComponentUObject->GetOwner()->GetLevel());
 		}
 
 		for (FRectLightBuildInfo& RectLight : LightScene.RectLights.Elements)
 		{
-			URectLightComponent* Light = RectLight.ComponentUObject;
-			check(!RectLight.bStationary || Light->PreviewShadowMapChannel != INDEX_NONE);
-
-			ULevel* StorageLevel = LightingScenario ? LightingScenario : Light->GetOwner()->GetLevel();
-			UMapBuildDataRegistry* Registry = StorageLevel->GetOrCreateMapBuildData();
-			FLightComponentMapBuildData& LightBuildData = Registry->FindOrAllocateLightBuildData(Light->LightGuid, true);
-			LightBuildData.ShadowMapChannel = RectLight.bStationary ? Light->PreviewShadowMapChannel : INDEX_NONE;
+			RectLight.AllocateMapBuildData(LightingScenario ? LightingScenario : RectLight.ComponentUObject->GetOwner()->GetLevel());
 		}
 
 		{
@@ -2222,7 +2208,7 @@ void FScene::ApplyFinishedLightmapsToWorld()
 								FLocalLightRenderState& Light
 								)
 							{
-								check(Light.bStationary);
+								check(Light.bStationary && Light.bCastShadow);
 								check(Light.ShadowMapChannel != INDEX_NONE);
 								FQuantizedShadowSignedDistanceFieldData2D* ShadowMap = new FQuantizedShadowSignedDistanceFieldData2D(Lightmap.GetSize().X, Lightmap.GetSize().Y);
 
@@ -2246,7 +2232,7 @@ void FScene::ApplyFinishedLightmapsToWorld()
 							// Directional lights are always relevant
 							for (FDirectionalLightBuildInfo& DirectionalLight : LightScene.DirectionalLights.Elements)
 							{
-								if (!DirectionalLight.bStationary)
+								if (!DirectionalLight.CastsStationaryShadow())
 								{
 									continue;
 								}
@@ -2447,7 +2433,7 @@ void FScene::ApplyFinishedLightmapsToWorld()
 								// Directional lights are always relevant
 								for (FDirectionalLightBuildInfo& DirectionalLight : LightScene.DirectionalLights.Elements)
 								{
-									if (!DirectionalLight.bStationary)
+									if (!DirectionalLight.CastsStationaryShadow())
 									{
 										continue;
 									}
@@ -2476,7 +2462,7 @@ void FScene::ApplyFinishedLightmapsToWorld()
 
 								for (FPointLightRenderStateRef& PointLight : Lightmap.RelevantPointLights)
 								{
-									check(PointLight->bStationary);
+									check(PointLight->bStationary && PointLight->bCastShadow);
 									check(PointLight->ShadowMapChannel != INDEX_NONE);
 									TUniquePtr<FQuantizedShadowSignedDistanceFieldData2D> ShadowMap = MakeUnique<FQuantizedShadowSignedDistanceFieldData2D>(BaseLightMapWidth, BaseLightMapHeight);
 
@@ -2501,7 +2487,7 @@ void FScene::ApplyFinishedLightmapsToWorld()
 
 								for (FSpotLightRenderStateRef& SpotLight : Lightmap.RelevantSpotLights)
 								{
-									check(SpotLight->bStationary);
+									check(SpotLight->bStationary && SpotLight->bCastShadow);
 									check(SpotLight->ShadowMapChannel != INDEX_NONE);
 									TUniquePtr<FQuantizedShadowSignedDistanceFieldData2D> ShadowMap = MakeUnique<FQuantizedShadowSignedDistanceFieldData2D>(BaseLightMapWidth, BaseLightMapHeight);
 
@@ -2526,7 +2512,7 @@ void FScene::ApplyFinishedLightmapsToWorld()
 
 								for (FRectLightRenderStateRef& RectLight : Lightmap.RelevantRectLights)
 								{
-									check(RectLight->bStationary);
+									check(RectLight->bStationary && RectLight->bCastShadow);
 									check(RectLight->ShadowMapChannel != INDEX_NONE);
 									TUniquePtr<FQuantizedShadowSignedDistanceFieldData2D> ShadowMap = MakeUnique<FQuantizedShadowSignedDistanceFieldData2D>(BaseLightMapWidth, BaseLightMapHeight);
 
@@ -2785,7 +2771,7 @@ void FScene::ApplyFinishedLightmapsToWorld()
 								FLocalLightRenderState& Light
 								)
 							{
-								check(Light.bStationary);
+								check(Light.bStationary && Light.bCastShadow);
 								check(Light.ShadowMapChannel != INDEX_NONE);
 								FQuantizedShadowSignedDistanceFieldData2D* ShadowMap = new FQuantizedShadowSignedDistanceFieldData2D(Lightmap.GetSize().X, Lightmap.GetSize().Y);
 
@@ -2809,7 +2795,7 @@ void FScene::ApplyFinishedLightmapsToWorld()
 							// Directional lights are always relevant
 							for (FDirectionalLightBuildInfo& DirectionalLight : LightScene.DirectionalLights.Elements)
 							{
-								if (!DirectionalLight.bStationary)
+								if (!DirectionalLight.CastsStationaryShadow())
 								{
 									continue;
 								}
