GPULM: Properly integrate direct lighting SH as we have 2 directions now (1 from BRDF sampling and 1 from NEE). Luminance for first bounce ray guiding is accumulated separately

diff --git a/Engine/Plugins/Experimental/GPULightmass/Shaders/Private/LightmapPathTracing.usf b/Engine/Plugins/Experimental/GPULightmass/Shaders/Private/LightmapPathTracing.usf
--- a/Engine/Plugins/Experimental/GPULightmass/Shaders/Private/LightmapPathTracing.usf
+++ b/Engine/Plugins/Experimental/GPULightmass/Shaders/Private/LightmapPathTracing.usf
@@ -316,8 +316,13 @@ void PathTracingKernel(
 	inout bool bIsValidSample,
 	inout float3 RadianceValue,
 	inout float3 RadianceDirection,
-	inout float3 DirectLightingRadianceValue,
-	inout float3 DirectLightingRadianceDirection,
+	inout float3 DirectLightingIrradiance,
+#ifdef LIGHTMAP_PATH_TRACING_MAIN_RG
+	inout FL2SHAndCorrection DirectLightingSH,
+#else
+	inout FThreeBandSHVectorRGBFloat DirectLightingSH,
+#endif
+	inout float LuminanceForFirstBounceRayGuiding,
 	inout float4 PrimaryRandSample)
 {
 	uint2 LaunchIndex = uint2(0, 0);
@@ -466,6 +471,7 @@ void PathTracingKernel(
 			{
 				// Successful query. Terminate path immediately.
 				Radiance += RecordIrradiance * PathThroughput;
+				LuminanceForFirstBounceRayGuiding += Luminance(RecordIrradiance * PathThroughput);
 				break;
 			}
 			else
@@ -598,14 +604,28 @@ void PathTracingKernel(
 						if (Bounce > 0)
 						{
 							AccumulateRadiance(Radiance, LightContrib);
+							LuminanceForFirstBounceRayGuiding += Luminance(LightContrib);
 						}
 						else
 						{
 							// GPU Lightmass records contribution of direct lighting separately, and only for static lights
 							if (!IsStationary(LightId))
 							{
-								DirectLightingRadianceValue += LightContrib;
-								DirectLightingRadianceDirection += MISWeightedDirection;
+							#ifdef LIGHTMAP_PATH_TRACING_MAIN_RG
+								float TangentZ = saturate(dot(LightSample.Direction, Payload.WorldNormal));
+								DirectLightingSH.AddIncomingRadiance(Luminance(LightContrib), LightSample.Direction, TangentZ);
+								DirectLightingIrradiance += LightContrib * TangentZ;
+							#else								
+								DirectLightingSH.R.V0 = SHBasisFunction3Float(LightSample.Direction).V0 * LightContrib.r;
+								DirectLightingSH.R.V1 = SHBasisFunction3Float(LightSample.Direction).V1 * LightContrib.r;
+								DirectLightingSH.R.V2 = SHBasisFunction3Float(LightSample.Direction).V2 * LightContrib.r;
+								DirectLightingSH.G.V0 = SHBasisFunction3Float(LightSample.Direction).V0 * LightContrib.g;
+								DirectLightingSH.G.V1 = SHBasisFunction3Float(LightSample.Direction).V1 * LightContrib.g;
+								DirectLightingSH.G.V2 = SHBasisFunction3Float(LightSample.Direction).V2 * LightContrib.g;
+								DirectLightingSH.B.V0 = SHBasisFunction3Float(LightSample.Direction).V0 * LightContrib.b;
+								DirectLightingSH.B.V1 = SHBasisFunction3Float(LightSample.Direction).V1 * LightContrib.b;
+								DirectLightingSH.B.V2 = SHBasisFunction3Float(LightSample.Direction).V2 * LightContrib.b;
+							#endif
 							}
 						}
 					}
@@ -707,12 +727,6 @@ void PathTracingKernel(
 					continue;
 				}
 
-				// GPU Lightmass records contribution of direct lighting separately, and only for static lights
-				if (Bounce == 0 && IsStationary(LightId))
-				{
-					continue;
-				}
-
 				float3 LightContrib = PathThroughput * LightResult.Radiance;
 				float3 MISWeightedDirection = MaterialSample.Direction;
 				
@@ -743,15 +757,35 @@ void PathTracingKernel(
 						// Separate direct lighting (bounce 0) for GPU Lightmass
 						MISWeightedDirection *= Visibility;
 					}
-					// the light made some contribution, and there was nothing along the shadow ray
-					AccumulateRadiance(Radiance, LightContrib);
 
-					// Separate direct lighting (bounce 0) for GPU Lightmass
-					if (Bounce == 0)
+					if (Bounce > 0)
 					{
-						DirectLightingRadianceValue += LightContrib;
-						DirectLightingRadianceDirection += MISWeightedDirection;
+						// the light made some contribution, and there was nothing along the shadow ray
+						AccumulateRadiance(Radiance, LightContrib);
 					}
+					else // Bounce == 0
+					{
+						if (!IsStationary(LightId))
+						{
+						#ifdef LIGHTMAP_PATH_TRACING_MAIN_RG
+							float TangentZ = saturate(dot(MaterialSample.Direction, Payload.WorldNormal));
+							DirectLightingSH.AddIncomingRadiance(Luminance(LightContrib), MaterialSample.Direction, TangentZ);
+							DirectLightingIrradiance += LightContrib * TangentZ;
+						#else								
+							DirectLightingSH.R.V0 = SHBasisFunction3Float(MaterialSample.Direction).V0 * LightContrib.r;
+							DirectLightingSH.R.V1 = SHBasisFunction3Float(MaterialSample.Direction).V1 * LightContrib.r;
+							DirectLightingSH.R.V2 = SHBasisFunction3Float(MaterialSample.Direction).V2 * LightContrib.r;
+							DirectLightingSH.G.V0 = SHBasisFunction3Float(MaterialSample.Direction).V0 * LightContrib.g;
+							DirectLightingSH.G.V1 = SHBasisFunction3Float(MaterialSample.Direction).V1 * LightContrib.g;
+							DirectLightingSH.G.V2 = SHBasisFunction3Float(MaterialSample.Direction).V2 * LightContrib.g;
+							DirectLightingSH.B.V0 = SHBasisFunction3Float(MaterialSample.Direction).V0 * LightContrib.b;
+							DirectLightingSH.B.V1 = SHBasisFunction3Float(MaterialSample.Direction).V1 * LightContrib.b;
+							DirectLightingSH.B.V2 = SHBasisFunction3Float(MaterialSample.Direction).V2 * LightContrib.b;
+						#endif
+						}
+					}
+					
+					LuminanceForFirstBounceRayGuiding += Luminance(LightContrib);
 				}
 			}
 			
@@ -847,8 +881,13 @@ void LightmapPathTracingMainRG()
 
 				float3 RadianceValue = 0;
 				float3 RadianceDirection = 0;				
-				float3 DirectLightingRadianceValue = 0;
-				float3 DirectLightingRadianceDirection = 0;
+				float3 DirectLightingIrradiance = 0;
+			#ifdef LIGHTMAP_PATH_TRACING_MAIN_RG
+				FL2SHAndCorrection DirectLightingSH = (FL2SHAndCorrection)0;
+			#else
+				FThreeBandSHVectorRGBFloat DirectLightingSH = (FThreeBandSHVectorRGBFloat)0;
+			#endif
+				float LuminanceForFirstBounceRayGuiding = 0;
 				
 				if (bIsValidSample)
 				{
@@ -861,8 +900,9 @@ void LightmapPathTracingMainRG()
 						bIsValidSample,
 						RadianceValue,
 						RadianceDirection,
-						DirectLightingRadianceValue,
-						DirectLightingRadianceDirection,
+						DirectLightingIrradiance,
+						DirectLightingSH,
+						LuminanceForFirstBounceRayGuiding,
 						PrimaryRandSample
 					);
 				}
@@ -885,14 +925,17 @@ void LightmapPathTracingMainRG()
 					
 					if (BatchedTiles[TileIndex].RenderPassIndex >= MinRenderPassIndex && BatchedTiles[TileIndex].RenderPassIndex < MaxRenderPassIndex)
 					{
+						float2 Point = PrimaryRandSample.yx;						
 						int2 ClusterPosition = clamp((int2)LaunchIndex - int2(2, 2), int2(0, 0), int2(63, 63)) / TEXEL_CLUSTER_SIZE;
-						float2 JitteredBin = PrimaryRandSample.xy * DIRECTIONAL_BINS_ONE_DIM;// - float2(0.5f, 0.5f) + PrimaryRandSample.zw;
+						float2 JitteredBin = Point * DIRECTIONAL_BINS_ONE_DIM;// - float2(0.5f, 0.5f) + PrimaryRandSample.xy;
 						int2 PositionInBin = clamp(JitteredBin, float2(0, 0), float2(DIRECTIONAL_BINS_ONE_DIM - 1, DIRECTIONAL_BINS_ONE_DIM - 1));
 						int2 FinalPosition = ClusterPosition * DIRECTIONAL_BINS_ONE_DIM + PositionInBin;
-					
-						float Illuminance = Luminance(RadianceValue) * saturate(dot(RadianceDirection, ShadingNormal));
-						// For positive floats we can cast them to uint and use atomic max directly
-						InterlockedMax(RayGuidingLuminance[BatchedTiles[TileIndex].WorkingSetPosition / GPreviewLightmapPhysicalTileSize * CDF_TILE_SIZE + FinalPosition], asuint(max(Illuminance, 0)));
+
+						{
+							float Illuminance = LuminanceForFirstBounceRayGuiding * saturate(dot(RadianceDirection, ShadingNormal));
+							// For positive floats we can cast them to uint and use atomic max directly
+							InterlockedMax(RayGuidingLuminance[BatchedTiles[TileIndex].WorkingSetPosition / GPreviewLightmapPhysicalTileSize * CDF_TILE_SIZE + FinalPosition], asuint(max(Illuminance, 0)));
+						}
 					}
 				#endif
 					
@@ -906,17 +949,14 @@ void LightmapPathTracingMainRG()
 						SHDirectionality[TexelIndexInPool] += SH.L2SHCoefficients;					
 						SHCorrectionAndStationarySkyLightBentNormal[TexelIndexInPool].x += SH.Correction;
 					}
-					
+
+					#ifdef LIGHTMAP_PATH_TRACING_MAIN_RG
 					{
-						FL2SHAndCorrection SH = (FL2SHAndCorrection)0;
-						
-						float TangentZ = saturate(dot(DirectLightingRadianceDirection, ShadingNormal));
-						SH.AddIncomingRadiance(Luminance(DirectLightingRadianceValue), DirectLightingRadianceDirection, TangentZ);
-						
-						IrradianceAndSampleCount[TexelIndexInPool].rgb += float3(DirectLightingRadianceValue * TangentZ / PI);
-						SHDirectionality[TexelIndexInPool] += SH.L2SHCoefficients;					
-						SHCorrectionAndStationarySkyLightBentNormal[TexelIndexInPool].x += SH.Correction;
+						IrradianceAndSampleCount[TexelIndexInPool].rgb += DirectLightingIrradiance / PI;
+						SHDirectionality[TexelIndexInPool] += DirectLightingSH.L2SHCoefficients;					
+						SHCorrectionAndStationarySkyLightBentNormal[TexelIndexInPool].x += DirectLightingSH.Correction;
 					}
+					#endif
 				}
 			}
 		}
@@ -1164,8 +1204,13 @@ void VolumetricLightmapPathTracingMainRG()
 
 	float3 RadianceValue = 0;
 	float3 RadianceDirection = 0;				
-	float3 DirectLightingRadianceValue = 0;
-	float3 DirectLightingRadianceDirection = 0;
+	float3 DirectLightingIrradiance = 0;
+#ifdef LIGHTMAP_PATH_TRACING_MAIN_RG
+	FL2SHAndCorrection DirectLightingSH = (FL2SHAndCorrection)0;
+#else
+	FThreeBandSHVectorRGBFloat DirectLightingSH = (FThreeBandSHVectorRGBFloat)0;
+#endif
+	float LuminanceForFirstBounceRayGuiding = 0;
 	float4 PrimaryRandSample;
 
 	if (bIsValidSample)
@@ -1183,15 +1228,26 @@ void VolumetricLightmapPathTracingMainRG()
 			bIsValidSample,
 			RadianceValue,
 			RadianceDirection,
-			DirectLightingRadianceValue,
-			DirectLightingRadianceDirection,
+			DirectLightingIrradiance,
+			DirectLightingSH,
+			LuminanceForFirstBounceRayGuiding,
 			PrimaryRandSample);
 	}
 	
 	// Every even sample is upper hemisphere, odd sample is lower hemisphere
 	// SampleCount stored is twice as the actual full sphere sample count
 	RadianceValue *= 2;
-	DirectLightingRadianceValue *= 2;
+#ifndef LIGHTMAP_PATH_TRACING_MAIN_RG
+	DirectLightingSH.R.V0 *= 2;
+	DirectLightingSH.R.V1 *= 2;
+	DirectLightingSH.R.V2 *= 2;
+	DirectLightingSH.G.V0 *= 2;
+	DirectLightingSH.G.V1 *= 2;
+	DirectLightingSH.G.V2 *= 2;
+	DirectLightingSH.B.V0 *= 2;
+	DirectLightingSH.B.V1 *= 2;
+	DirectLightingSH.B.V2 *= 2;
+#endif
 
 	if (any(isnan(RadianceValue)) || any(RadianceValue < 0))
 	{
@@ -1244,32 +1300,23 @@ void VolumetricLightmapPathTracingMainRG()
 			SHCoefficients0B[VoxelPos] += float4(SH.B.V0.yzw, SH.B.V1.x);
 			SHCoefficients1B[VoxelPos] += float4(SH.B.V1.yzw, SH.B.V2);
 		}
-		
+
+		#ifndef LIGHTMAP_PATH_TRACING_MAIN_RG
 		{
-			FThreeBandSHVectorRGBFloat SH;
-			SH.R.V0 = SHBasisFunction3Float(DirectLightingRadianceDirection).V0 * DirectLightingRadianceValue.r;
-			SH.R.V1 = SHBasisFunction3Float(DirectLightingRadianceDirection).V1 * DirectLightingRadianceValue.r;
-			SH.R.V2 = SHBasisFunction3Float(DirectLightingRadianceDirection).V2 * DirectLightingRadianceValue.r;
-			SH.G.V0 = SHBasisFunction3Float(DirectLightingRadianceDirection).V0 * DirectLightingRadianceValue.g;
-			SH.G.V1 = SHBasisFunction3Float(DirectLightingRadianceDirection).V1 * DirectLightingRadianceValue.g;
-			SH.G.V2 = SHBasisFunction3Float(DirectLightingRadianceDirection).V2 * DirectLightingRadianceValue.g;
-			SH.B.V0 = SHBasisFunction3Float(DirectLightingRadianceDirection).V0 * DirectLightingRadianceValue.b;
-			SH.B.V1 = SHBasisFunction3Float(DirectLightingRadianceDirection).V1 * DirectLightingRadianceValue.b;
-			SH.B.V2 = SHBasisFunction3Float(DirectLightingRadianceDirection).V2 * DirectLightingRadianceValue.b;
+			AmbientVector[VoxelPos].x += DirectLightingSH.R.V0.x;
+			AmbientVector[VoxelPos].y += DirectLightingSH.G.V0.x;
+			AmbientVector[VoxelPos].z += DirectLightingSH.B.V0.x;
 			
-			AmbientVector[VoxelPos].x += SH.R.V0.x;
-			AmbientVector[VoxelPos].y += SH.G.V0.x;
-			AmbientVector[VoxelPos].z += SH.B.V0.x;
-			
-			SHCoefficients0R[VoxelPos] += float4(SH.R.V0.yzw, SH.R.V1.x);
-			SHCoefficients1R[VoxelPos] += float4(SH.R.V1.yzw, SH.R.V2);
+			SHCoefficients0R[VoxelPos] += float4(DirectLightingSH.R.V0.yzw, DirectLightingSH.R.V1.x);
+			SHCoefficients1R[VoxelPos] += float4(DirectLightingSH.R.V1.yzw, DirectLightingSH.R.V2);
 			
-			SHCoefficients0G[VoxelPos] += float4(SH.G.V0.yzw, SH.G.V1.x);
-			SHCoefficients1G[VoxelPos] += float4(SH.G.V1.yzw, SH.G.V2);		
+			SHCoefficients0G[VoxelPos] += float4(DirectLightingSH.G.V0.yzw, DirectLightingSH.G.V1.x);
+			SHCoefficients1G[VoxelPos] += float4(DirectLightingSH.G.V1.yzw, DirectLightingSH.G.V2);		
 			
-			SHCoefficients0B[VoxelPos] += float4(SH.B.V0.yzw, SH.B.V1.x);
-			SHCoefficients1B[VoxelPos] += float4(SH.B.V1.yzw, SH.B.V2);
+			SHCoefficients0B[VoxelPos] += float4(DirectLightingSH.B.V0.yzw, DirectLightingSH.B.V1.x);
+			SHCoefficients1B[VoxelPos] += float4(DirectLightingSH.B.V1.yzw, DirectLightingSH.B.V2);
 		}
+		#endif
 		
 		AmbientVector[VoxelPos].w = asfloat(SampleCount + 1);
 	}
diff --git a/Engine/Shaders/Private/PathTracing/Material/PathTracingRadianceProbe.ush b/Engine/Shaders/Private/PathTracing/Material/PathTracingRadianceProbe.ush
--- a/Engine/Shaders/Private/PathTracing/Material/PathTracingRadianceProbe.ush
+++ b/Engine/Shaders/Private/PathTracing/Material/PathTracingRadianceProbe.ush
@@ -54,7 +54,9 @@ FMaterialSample RadianceProbe_SampleMaterial(
 )
 {
 	float PdfAdjust = 1.0f; 
-	
+
+	float2 Point = RandSample.yx;
+
 #ifdef LIGHTMAP_PATH_TRACING_MAIN_RG
 #if USE_FIRST_BOUNCE_RAY_GUIDING
 	uint2 BatchedLaunchIndex = DispatchRaysIndex().xy;
@@ -86,9 +88,9 @@ FMaterialSample RadianceProbe_SampleMaterial(
 			LastRowPrefixSum = RowPrefixSum;
 			RowPrefixSum = RayGuidingCDFY[BatchedTiles[TileIndex].WorkingSetPosition / GPreviewLightmapPhysicalTileSize * CDF_TILE_SIZE / 4 + FinalPosition];
 			
-			if (RowPrefixSum > RandSample.y)
+			if (RowPrefixSum > Point.y)
 			{
-				RandSample.y = (y + (RandSample.y - LastRowPrefixSum) / (RowPrefixSum - LastRowPrefixSum)) / DIRECTIONAL_BINS_ONE_DIM;
+				Point.y = (y + (Point.y - LastRowPrefixSum) / (RowPrefixSum - LastRowPrefixSum)) / DIRECTIONAL_BINS_ONE_DIM;
 				
 				PdfAdjust *= (RowPrefixSum - LastRowPrefixSum) * DIRECTIONAL_BINS_ONE_DIM;
 				
@@ -108,9 +110,9 @@ FMaterialSample RadianceProbe_SampleMaterial(
 			LastPrefixSum = PrefixSum;
 			PrefixSum = RayGuidingCDFX[BatchedTiles[TileIndex].WorkingSetPosition / GPreviewLightmapPhysicalTileSize * CDF_TILE_SIZE + FinalPosition];
 			
-			if (PrefixSum > RandSample.x)
+			if (PrefixSum > Point.x)
 			{
-				RandSample.x = (x + (RandSample.x - LastPrefixSum) / (PrefixSum - LastPrefixSum)) / DIRECTIONAL_BINS_ONE_DIM;									
+				Point.x = (x + (Point.x - LastPrefixSum) / (PrefixSum - LastPrefixSum)) / DIRECTIONAL_BINS_ONE_DIM;									
 				PdfAdjust *= (PrefixSum - LastPrefixSum) * DIRECTIONAL_BINS_ONE_DIM;
 				break;
 			}
@@ -123,14 +125,13 @@ FMaterialSample RadianceProbe_SampleMaterial(
 
 #ifdef LIGHTMAP_PATH_TRACING_MAIN_RG
 	// Use cosine hemisphere sampling for surface lightmaps
-	float4 SampledValue = CosineSampleHemisphereConcentric(RandSample.xy);
+	float4 SampledValue = CosineSampleHemisphereConcentric(Point);
 #else
 	// Use uniform hemisphere sampling for volumetric lightmaps (the two hemispheres are meant to be merged)
-	float4 SampledValue = UniformSampleHemisphere(RandSample.xy);
+	float4 SampledValue = UniformSampleHemisphere(Point);
 #endif
 	float3 OutDirection = TangentToWorld(SampledValue.xyz, N_World);
-
-	PdfAdjust = max(PdfAdjust, 0.00001f);
+	
 	float OutPdf = SampledValue.w * PdfAdjust;
 	// As the name suggests, radiance probe has a combined throughput (Pdf * Weight) == 1
 	float OutWeight = (OutPdf > 0.0f) ? (1.0f / OutPdf) : 0;
@@ -138,82 +139,6 @@ FMaterialSample RadianceProbe_SampleMaterial(
 	return CreateMaterialSample(OutDirection, OutWeight, OutPdf, 1.0, 1.0);
 }
 
-float RadianceProbe_Pdf(float3 OutgoingDirection, FPathTracingPayload Payload)
-{
-	float3 N_World = Payload.WorldNormal;
-	
-	float PdfAdjust = 1.0f;	
-#ifdef LIGHTMAP_PATH_TRACING_MAIN_RG
-#if USE_FIRST_BOUNCE_RAY_GUIDING
-	float3 TangentDirection = WorldToTangent(OutgoingDirection, N_World);
-	float2 PrimarySample = InverseConcentricMapping(TangentDirection.xy);
-	
-	uint2 BatchedLaunchIndex = DispatchRaysIndex().xy;
-
-	uint2 LaunchIndex = uint2(BatchedLaunchIndex.x % GPreviewLightmapPhysicalTileSize, BatchedLaunchIndex.y);
-	int TileIndex = BatchedLaunchIndex.x / GPreviewLightmapPhysicalTileSize;
-
-	int2 ClusterPosition = clamp((int2)LaunchIndex - int2(2, 2), int2(0, 0), int2(63, 63)) / TEXEL_CLUSTER_SIZE;
-
-	int MinRenderPassIndex = NumRayGuidingTrialSamples;
-
-#if USE_IRRADIANCE_CACHING
-	MinRenderPassIndex += GetIrradianceCachingQuality();
-#endif
-	
-	if (BatchedTiles[TileIndex].RenderPassIndex >= MinRenderPassIndex)
-	{
-		float LastPrefixSum = 0;
-		float PrefixSum;
-		float LastRowPrefixSum = 0;
-		float RowPrefixSum;
-		
-		int Y = floor(PrimarySample.y * DIRECTIONAL_BINS_ONE_DIM);	
-		
-		{
-			int2 WritePos = int2(Y % 4, Y / 4);
-			int2 FinalPosition = ClusterPosition * DIRECTIONAL_BINS_ONE_DIM / 4 + WritePos;
-			
-			RowPrefixSum = RayGuidingCDFY[BatchedTiles[TileIndex].WorkingSetPosition / GPreviewLightmapPhysicalTileSize * CDF_TILE_SIZE / 4 + FinalPosition];
-		}
-		
-		if (Y > 0)
-		{
-			int2 WritePos = int2((Y - 1) % 4, (Y - 1) / 4);
-			int2 FinalPosition = ClusterPosition * DIRECTIONAL_BINS_ONE_DIM / 4 + WritePos;
-			
-			LastRowPrefixSum = RayGuidingCDFY[BatchedTiles[TileIndex].WorkingSetPosition / GPreviewLightmapPhysicalTileSize * CDF_TILE_SIZE / 4 + FinalPosition];
-		}
-		
-		PdfAdjust *= (RowPrefixSum - LastRowPrefixSum) * DIRECTIONAL_BINS_ONE_DIM;	
-		
-		int X = floor(PrimarySample.x * DIRECTIONAL_BINS_ONE_DIM);
-		
-		{
-			int2 FinalPosition = ClusterPosition * DIRECTIONAL_BINS_ONE_DIM + int2(X, Y);
-			PrefixSum = RayGuidingCDFX[BatchedTiles[TileIndex].WorkingSetPosition / GPreviewLightmapPhysicalTileSize * CDF_TILE_SIZE + FinalPosition];
-		}
-		
-		if (X > 0)
-		{
-			int2 FinalPosition = ClusterPosition * DIRECTIONAL_BINS_ONE_DIM + int2(X - 1, Y);
-			LastPrefixSum = RayGuidingCDFX[BatchedTiles[TileIndex].WorkingSetPosition / GPreviewLightmapPhysicalTileSize * CDF_TILE_SIZE + FinalPosition];
-		}
-		
-		PdfAdjust *= (PrefixSum - LastPrefixSum) * DIRECTIONAL_BINS_ONE_DIM;
-	}
-#endif	
-#endif
-
-	PdfAdjust = max(PdfAdjust, 0.00001f);
-#ifdef LIGHTMAP_PATH_TRACING_MAIN_RG
-	float NoL = max(dot(N_World, OutgoingDirection), 0.0);
-	return (NoL / PI) * PdfAdjust;
-#else
-	return 1 / (2 * PI) * PdfAdjust;
-#endif
-}
-
 FMaterialEval RadianceProbe_EvalMaterial(
 	float3 OutgoingDirection,
 	FPathTracingPayload Payload
@@ -284,9 +209,7 @@ FMaterialEval RadianceProbe_EvalMaterial(
 	}
 #endif	
 #endif
-	float NoL = max(dot(N_World, OutgoingDirection), 0.0);
-
-	PdfAdjust = max(PdfAdjust, 0.00001f);
+	float NoL = max(dot(N_World, OutgoingDirection), 0);
 
 #ifdef LIGHTMAP_PATH_TRACING_MAIN_RG
 	float OutPdf = NoL / PI * PdfAdjust;
