GPULM: Preliminary support for stationary skylights. Results cannot be saved now

diff --git a/Engine/Plugins/Experimental/GPULightmass/Shaders/Private/BatchedTiles.ush b/Engine/Plugins/Experimental/GPULightmass/Shaders/Private/BatchedTiles.ush
--- a/Engine/Plugins/Experimental/GPULightmass/Shaders/Private/BatchedTiles.ush
+++ b/Engine/Plugins/Experimental/GPULightmass/Shaders/Private/BatchedTiles.ush
@@ -7,6 +7,7 @@ struct FGPUTileDescription
 	int2 OutputLayer0Position;
 	int2 OutputLayer1Position;
 	int2 OutputLayer2Position;
+	int2 OutputLayer3Position;
 	int FrameIndex;
 	int RenderPassIndex;
 };
diff --git a/Engine/Plugins/Experimental/GPULightmass/Shaders/Private/FirstBounceRayGuidingCDFBuild.usf b/Engine/Plugins/Experimental/GPULightmass/Shaders/Private/FirstBounceRayGuidingCDFBuild.usf
--- a/Engine/Plugins/Experimental/GPULightmass/Shaders/Private/FirstBounceRayGuidingCDFBuild.usf
+++ b/Engine/Plugins/Experimental/GPULightmass/Shaders/Private/FirstBounceRayGuidingCDFBuild.usf
@@ -34,7 +34,7 @@ void FirstBounceRayGuidingCDFBuildCS(
 			for(int dy = -FilterKernelSize; dy <= FilterKernelSize; dy++)
 			{
 				int2 FinalPosition = ClusterPosition * DIRECTIONAL_BINS_ONE_DIM + clamp(ThreadId.xy + int2(dx, dy), int2(0, 0), int2(DIRECTIONAL_BINS_ONE_DIM - 1, DIRECTIONAL_BINS_ONE_DIM - 1));
-				Value += asfloat(RayGuidingLuminance[BatchedTiles[TileIndex].WorkingSetPosition / GPreviewLightmapPhysicalTileSize * CDF_TILE_SIZE + FinalPosition]) * ((dx != 0 || dy != 0) ? 0.125 : 1);
+				Value += max(asfloat(RayGuidingLuminance[BatchedTiles[TileIndex].WorkingSetPosition / GPreviewLightmapPhysicalTileSize * CDF_TILE_SIZE + FinalPosition]) * ((dx != 0 || dy != 0) ? 0.125 : 1), 0.00001f);
 			}
 		}
 
diff --git a/Engine/Plugins/Experimental/GPULightmass/Shaders/Private/LightmapBufferClear.usf b/Engine/Plugins/Experimental/GPULightmass/Shaders/Private/LightmapBufferClear.usf
--- a/Engine/Plugins/Experimental/GPULightmass/Shaders/Private/LightmapBufferClear.usf
+++ b/Engine/Plugins/Experimental/GPULightmass/Shaders/Private/LightmapBufferClear.usf
@@ -25,6 +25,7 @@ uint StagingPoolSizeX;
 RWTexture2D<float4> StagingHQLayer0;
 RWTexture2D<float4> StagingHQLayer1;
 RWTexture2D<float4> StagingShadowMask;
+RWTexture2D<float4> StagingSkyOcclusion;
 
 [numthreads(8, 8, 1)]
 void CopyConvergedLightmapTilesCS(uint3 DispatchThreadId : SV_DispatchThreadID)
@@ -69,6 +70,20 @@ void CopyConvergedLightmapTilesCS(uint3 DispatchThreadId : SV_DispatchThreadID)
 
 		StagingShadowMask[TexelIndexInStagingPool] = sqrt((float4)ShadowValue / max(uint4(1, 1, 1, 1), SampleCountValue)) * ValidityMask;
 	}
+
+	{
+		uint SampleCount = asuint(IrradianceAndSampleCount[TexelIndexInPool].w);
+		
+		if (SampleCount > 0)
+		{
+			StagingSkyOcclusion[TexelIndexInStagingPool].rgb = saturate(normalize(SHCorrectionAndStationarySkyLightBentNormal[TexelIndexInPool].yzw / SampleCount) * 0.5 + 0.5);
+			StagingSkyOcclusion[TexelIndexInStagingPool].a = saturate(sqrt(length(SHCorrectionAndStationarySkyLightBentNormal[TexelIndexInPool].yzw / SampleCount)));
+		}
+		else
+		{
+			StagingSkyOcclusion[TexelIndexInStagingPool] = float4(0, 0, 0, 0);
+		}
+	}
 }
 
 RWTexture2D<float4> SrcTexture;
diff --git a/Engine/Plugins/Experimental/GPULightmass/Shaders/Private/LightmapOutput.usf b/Engine/Plugins/Experimental/GPULightmass/Shaders/Private/LightmapOutput.usf
--- a/Engine/Plugins/Experimental/GPULightmass/Shaders/Private/LightmapOutput.usf
+++ b/Engine/Plugins/Experimental/GPULightmass/Shaders/Private/LightmapOutput.usf
@@ -96,5 +96,14 @@ void SelectiveLightmapOutputCS(uint3 DispatchThreadId : SV_DispatchThreadID)
 	}
 	else if (DIM_OUTPUT_LAYER == 3)
 	{
+		if (SampleCount > 0)
+		{
+			OutputTileAtlas[LaunchIndex + BatchedTiles[TileIndex].OutputLayer3Position].rgb = saturate(normalize(SHCorrectionAndStationarySkyLightBentNormal[TexelIndexInPool].yzw / SampleCount) * 0.5 + 0.5);
+			OutputTileAtlas[LaunchIndex + BatchedTiles[TileIndex].OutputLayer3Position].a = saturate(sqrt(length(SHCorrectionAndStationarySkyLightBentNormal[TexelIndexInPool].yzw / SampleCount)));
+		}
+		else
+		{
+			OutputTileAtlas[LaunchIndex + BatchedTiles[TileIndex].OutputLayer3Position] = float4(0, 0, 0, 0);
+		}
 	}
 }
diff --git a/Engine/Plugins/Experimental/GPULightmass/Shaders/Private/LightmapPathTracing.usf b/Engine/Plugins/Experimental/GPULightmass/Shaders/Private/LightmapPathTracing.usf
--- a/Engine/Plugins/Experimental/GPULightmass/Shaders/Private/LightmapPathTracing.usf
+++ b/Engine/Plugins/Experimental/GPULightmass/Shaders/Private/LightmapPathTracing.usf
@@ -323,6 +323,7 @@ void PathTracingKernel(
 	inout FThreeBandSHVectorRGBFloat DirectLightingSH,
 #endif
 	inout float LuminanceForFirstBounceRayGuiding,
+	inout float3 SkyLightBentNormal,
 	inout float4 PrimaryRandSample)
 {
 	uint2 LaunchIndex = uint2(0, 0);
@@ -568,6 +569,8 @@ void PathTracingKernel(
 				LightSample.Pdf *= LightPickPdf;
 				if (LightSample.Pdf > 0)
 				{
+					float3 Visibility = float3(1, 1, 1);
+					
 					if (CastsShadow(LightId))
 					{
 						// for transmissive materials, bias the position to the other side of the surface if the light is coming from behind
@@ -582,7 +585,8 @@ void PathTracingKernel(
 
 						float AvgRoughness = ApproximateCaustics ? GetAverageRoughness(Payload) : 0.0;
 
-						LightSample.RadianceOverPdf *= TraceTransparentVisibilityRay(LightRay, LaunchIndex, AvgRoughness);
+						Visibility = TraceTransparentVisibilityRay(LightRay, LaunchIndex, AvgRoughness);
+						LightSample.RadianceOverPdf *= Visibility;
 					}
 
 					// #dxr_todo: Is it cheaper to fire the ray first? Or eval the material first?
@@ -592,12 +596,12 @@ void PathTracingKernel(
 						FMaterialEval MaterialEval = EvalMaterial(Ray.Direction, LightSample.Direction, Payload, Bounce == 0);
 
 						// Record the contribution
-						float3 LightContrib = PathThroughput * LightSample.RadianceOverPdf * MaterialEval.Weight * MaterialEval.Pdf;
-						float3 MISWeightedDirection = LightSample.Direction;
+						float3 LightContrib = PathThroughput * LightSample.RadianceOverPdf * MaterialEval.Weight * MaterialEval.Pdf;						
+						float3 BentNormalVector = LightSample.Direction / LightSample.Pdf / PI * Visibility;
 						if (MISMode == 2)
 						{
 							LightContrib *= MISWeightRobust(LightSample.Pdf, MaterialEval.Pdf);
-							MISWeightedDirection *= MISWeightRobust(LightSample.Pdf, MaterialEval.Pdf);
+							BentNormalVector *= MISWeightRobust(LightSample.Pdf, MaterialEval.Pdf);
 						}
 
 						// Record the contribution
@@ -627,6 +631,11 @@ void PathTracingKernel(
 								DirectLightingSH.B.V2 = SHBasisFunction3Float(LightSample.Direction).V2 * LightContrib.b;
 							#endif
 							}
+
+							if (IsStationary(LightId) && IsEnvironmentLight(LightId) && CastsShadow(LightId))
+							{
+								SkyLightBentNormal += BentNormalVector;
+							}
 						}
 					}
 				}
@@ -728,7 +737,7 @@ void PathTracingKernel(
 				}
 
 				float3 LightContrib = PathThroughput * LightResult.Radiance;
-				float3 MISWeightedDirection = MaterialSample.Direction;
+				float3 BentNormalVector = PathThroughput * MaterialSample.Direction / PI;
 				
 				if (bUseMIS)
 				{
@@ -742,7 +751,7 @@ void PathTracingKernel(
 					LightContrib *= MISWeightRobust(MaterialSample.Pdf, LightResult.Pdf * LightPickPdf);
 
 					// Separate direct lighting (bounce 0) for GPU Lightmass
-					MISWeightedDirection *= MISWeightRobust(MaterialSample.Pdf, LightResult.Pdf * LightPickPdf);
+					BentNormalVector *= MISWeightRobust(MaterialSample.Pdf, LightResult.Pdf * LightPickPdf);
 				}
 
 				if (any(LightContrib > 0))
@@ -755,7 +764,7 @@ void PathTracingKernel(
 						LightContrib *= Visibility;
 						
 						// Separate direct lighting (bounce 0) for GPU Lightmass
-						MISWeightedDirection *= Visibility;
+						BentNormalVector *= Visibility;
 					}
 
 					if (Bounce > 0)
@@ -783,6 +792,11 @@ void PathTracingKernel(
 							DirectLightingSH.B.V2 = SHBasisFunction3Float(MaterialSample.Direction).V2 * LightContrib.b;
 						#endif
 						}
+
+						if (IsStationary(LightId) && IsEnvironmentLight(LightId) && CastsShadow(LightId))
+						{
+							SkyLightBentNormal += BentNormalVector;
+						}
 					}
 					
 					LuminanceForFirstBounceRayGuiding += Luminance(LightContrib);
@@ -860,7 +874,7 @@ void LightmapPathTracingMainRG()
 	{
 		IrradianceAndSampleCount[TexelIndexInPool]= float4(0, 0, 0, 0);
 		SHDirectionality[TexelIndexInPool] = float4(0, 0, 0, 0);
-		SHCorrectionAndStationarySkyLightBentNormal[TexelIndexInPool].x = 0;
+		SHCorrectionAndStationarySkyLightBentNormal[TexelIndexInPool] = float4(0, 0, 0, 0);
 	}
 	
 	if (NumTotalSamples <= 0 || (NumTotalSamples > 0 && BatchedTiles[TileIndex].RenderPassIndex < NumTotalSamples - 1))
@@ -888,6 +902,7 @@ void LightmapPathTracingMainRG()
 				FThreeBandSHVectorRGBFloat DirectLightingSH = (FThreeBandSHVectorRGBFloat)0;
 			#endif
 				float LuminanceForFirstBounceRayGuiding = 0;
+				float3 SkyLightBentNormal = 0;
 				
 				if (bIsValidSample)
 				{
@@ -903,6 +918,7 @@ void LightmapPathTracingMainRG()
 						DirectLightingIrradiance,
 						DirectLightingSH,
 						LuminanceForFirstBounceRayGuiding,
+						SkyLightBentNormal,
 						PrimaryRandSample
 					);
 				}
@@ -957,6 +973,8 @@ void LightmapPathTracingMainRG()
 						SHCorrectionAndStationarySkyLightBentNormal[TexelIndexInPool].x += DirectLightingSH.Correction;
 					}
 					#endif
+
+					SHCorrectionAndStationarySkyLightBentNormal[TexelIndexInPool].yzw += SkyLightBentNormal;
 				}
 			}
 		}
@@ -1211,6 +1229,7 @@ void VolumetricLightmapPathTracingMainRG()
 	FThreeBandSHVectorRGBFloat DirectLightingSH = (FThreeBandSHVectorRGBFloat)0;
 #endif
 	float LuminanceForFirstBounceRayGuiding = 0;
+	float3 SkyLightBentNormal = 0;
 	float4 PrimaryRandSample;
 
 	if (bIsValidSample)
@@ -1231,6 +1250,7 @@ void VolumetricLightmapPathTracingMainRG()
 			DirectLightingIrradiance,
 			DirectLightingSH,
 			LuminanceForFirstBounceRayGuiding,
+			SkyLightBentNormal,
 			PrimaryRandSample);
 	}
 	
diff --git a/Engine/Plugins/Experimental/GPULightmass/Source/GPULightmass/Private/LightmapPreviewVirtualTexture.cpp b/Engine/Plugins/Experimental/GPULightmass/Source/GPULightmass/Private/LightmapPreviewVirtualTexture.cpp
--- a/Engine/Plugins/Experimental/GPULightmass/Source/GPULightmass/Private/LightmapPreviewVirtualTexture.cpp
+++ b/Engine/Plugins/Experimental/GPULightmass/Source/GPULightmass/Private/LightmapPreviewVirtualTexture.cpp
@@ -26,14 +26,16 @@ FLightmapPreviewVirtualTexture::FLightmapPreviewVirtualTexture(FLightmapRenderSt
 		ProducerDesc.BlockWidthInTiles = SizeInTiles.X;
 		ProducerDesc.BlockHeightInTiles = SizeInTiles.Y;
 		ProducerDesc.DepthInTiles = 1;
-		ProducerDesc.NumTextureLayers = 3; // LightMapVirtualTexture->TypeToLayer.Num();
-		ProducerDesc.NumPhysicalGroups = 3;
+		ProducerDesc.NumTextureLayers = 4; // LightMapVirtualTexture->TypeToLayer.Num();
+		ProducerDesc.NumPhysicalGroups = 4;
 		ProducerDesc.LayerFormat[0] = EPixelFormat::PF_A32B32G32R32F;
 		ProducerDesc.LayerFormat[1] = EPixelFormat::PF_A32B32G32R32F;
 		ProducerDesc.LayerFormat[2] = EPixelFormat::PF_A32B32G32R32F;
+		ProducerDesc.LayerFormat[3] = EPixelFormat::PF_A32B32G32R32F;
 		ProducerDesc.PhysicalGroupIndex[0] = 0;
 		ProducerDesc.PhysicalGroupIndex[1] = 1;
 		ProducerDesc.PhysicalGroupIndex[2] = 2;
+		ProducerDesc.PhysicalGroupIndex[3] = 3;
 		ProducerDesc.MaxLevel = FMath::Min((int32)FMath::CeilLogTwo((uint32)FMath::Min(SizeInTiles.X, SizeInTiles.Y)), GPreviewLightmapMipmapMaxLevel);
 
 		ProducerHandle = GetRendererModule().RegisterVirtualTextureProducer(ProducerDesc, this);
diff --git a/Engine/Plugins/Experimental/GPULightmass/Source/GPULightmass/Private/LightmapRenderer.cpp b/Engine/Plugins/Experimental/GPULightmass/Source/GPULightmass/Private/LightmapRenderer.cpp
--- a/Engine/Plugins/Experimental/GPULightmass/Source/GPULightmass/Private/LightmapRenderer.cpp
+++ b/Engine/Plugins/Experimental/GPULightmass/Source/GPULightmass/Private/LightmapRenderer.cpp
@@ -56,6 +56,7 @@ class FCopyConvergedLightmapTilesCS : public FGlobalShader
 		SHADER_PARAMETER_RDG_TEXTURE_UAV(RWTexture2D<float4>, StagingHQLayer0)
 		SHADER_PARAMETER_RDG_TEXTURE_UAV(RWTexture2D<float4>, StagingHQLayer1)
 		SHADER_PARAMETER_RDG_TEXTURE_UAV(RWTexture2D<float4>, StagingShadowMask)
+		SHADER_PARAMETER_RDG_TEXTURE_UAV(RWTexture2D<float4>, StagingSkyOcclusion)
 	END_SHADER_PARAMETER_STRUCT()
 };
 
@@ -157,6 +158,7 @@ struct FGPUTileDescription
 	FIntPoint OutputLayer0Position;
 	FIntPoint OutputLayer1Position;
 	FIntPoint OutputLayer2Position;
+	FIntPoint OutputLayer3Position;
 	int32 FrameIndex;
 	int32 RenderPassIndex;
 };
@@ -182,6 +184,7 @@ struct FGPUBatchedTileRequests
 			TileDesc.OutputLayer0Position = Tile.OutputPhysicalCoordinates[0] * GPreviewLightmapPhysicalTileSize;
 			TileDesc.OutputLayer1Position = Tile.OutputPhysicalCoordinates[1] * GPreviewLightmapPhysicalTileSize;
 			TileDesc.OutputLayer2Position = Tile.OutputPhysicalCoordinates[2] * GPreviewLightmapPhysicalTileSize;
+			TileDesc.OutputLayer3Position = Tile.OutputPhysicalCoordinates[3] * GPreviewLightmapPhysicalTileSize;
 			TileDesc.FrameIndex = Tile.RenderState->RetrieveTileState(Tile.VirtualCoordinates).Revision;
 			TileDesc.RenderPassIndex = Tile.RenderState->RetrieveTileState(Tile.VirtualCoordinates).RenderPassIndex;
 			BatchedTilesDesc.Add(TileDesc);
@@ -1488,16 +1491,18 @@ void FLightmapRenderer::Finalize(FRDGBuilder& GraphBuilder)
 			int32 NewSize = FMath::CeilToInt(FMath::Sqrt(static_cast<float>(TileUploadRequests.Num())));
 			if (!UploadTilePoolGPU.IsValid() || UploadTilePoolGPU->SizeInTiles.X < NewSize)
 			{
-				UploadTilePoolGPU = MakeUnique<FLightmapTilePoolGPU>(3, FIntPoint(NewSize, NewSize), FIntPoint(GPreviewLightmapPhysicalTileSize, GPreviewLightmapPhysicalTileSize));
+				UploadTilePoolGPU = MakeUnique<FLightmapTilePoolGPU>(4, FIntPoint(NewSize, NewSize), FIntPoint(GPreviewLightmapPhysicalTileSize, GPreviewLightmapPhysicalTileSize));
 				UE_LOG(LogGPULightmass, Log, TEXT("Resizing GPULightmass upload tile pool to (%d, %d) %dx%d"), NewSize, NewSize, NewSize * GPreviewLightmapPhysicalTileSize, NewSize * GPreviewLightmapPhysicalTileSize);
 			}
 
 			{
 				uint32 DstRowPitch;
-				FLinearColor* Texture[3];
-				Texture[0] = (FLinearColor*)RHICmdList.LockTexture2D(UploadTilePoolGPU->PooledRenderTargets[0]->GetRHI(), 0, RLM_WriteOnly, DstRowPitch, false);
-				Texture[1] = (FLinearColor*)RHICmdList.LockTexture2D(UploadTilePoolGPU->PooledRenderTargets[1]->GetRHI(), 0, RLM_WriteOnly, DstRowPitch, false);
-				Texture[2] = (FLinearColor*)RHICmdList.LockTexture2D(UploadTilePoolGPU->PooledRenderTargets[2]->GetRHI(), 0, RLM_WriteOnly, DstRowPitch, false);
+				FLinearColor* Texture[4];
+				check(UE_ARRAY_COUNT(Texture) == UploadTilePoolGPU->PooledRenderTargets.Num());
+				for (int32 LayerIndex = 0; LayerIndex < UE_ARRAY_COUNT(Texture); LayerIndex++)
+				{
+					Texture[LayerIndex] = (FLinearColor*)RHICmdList.LockTexture2D(UploadTilePoolGPU->PooledRenderTargets[LayerIndex]->GetRHI(), 0, RLM_WriteOnly, DstRowPitch, false);
+				}
 
 				TSet<FVirtualTile> TilesToDecompress;
 
@@ -1543,9 +1548,10 @@ void FLightmapRenderer::Finalize(FRDGBuilder& GraphBuilder)
 
 						if (!bLoadingOutOfBounds)
 						{
-							Tile.RenderState->TileStorage[SrcTileCoords].CPUTextureData[0]->Decompress();
-							Tile.RenderState->TileStorage[SrcTileCoords].CPUTextureData[1]->Decompress();
-							Tile.RenderState->TileStorage[SrcTileCoords].CPUTextureData[2]->Decompress();
+							for (int32 LayerIndex = 0; LayerIndex < UE_ARRAY_COUNT(Texture); LayerIndex++)
+							{
+								Tile.RenderState->TileStorage[SrcTileCoords].CPUTextureData[LayerIndex]->Decompress();
+							}
 						}
 					}
 				}
@@ -1590,27 +1596,29 @@ void FLightmapRenderer::Finalize(FRDGBuilder& GraphBuilder)
 								}
 							}
 
-							Texture[0][DstLinearIndex] = !bLoadingOutOfBounds ? TileUploadRequests[TileIndex].RenderState->TileStorage[SrcTileCoords].CPUTextureData[0]->Data[SrcLinearIndex] : FLinearColor(0, 0, 0, 0);
-							Texture[1][DstLinearIndex] = !bLoadingOutOfBounds ? TileUploadRequests[TileIndex].RenderState->TileStorage[SrcTileCoords].CPUTextureData[1]->Data[SrcLinearIndex] : FLinearColor(0, 0, 0, 0);
-							Texture[2][DstLinearIndex] = !bLoadingOutOfBounds ? TileUploadRequests[TileIndex].RenderState->TileStorage[SrcTileCoords].CPUTextureData[2]->Data[SrcLinearIndex] : FLinearColor(0, 0, 0, 0);
+							for (int32 LayerIndex = 0; LayerIndex < UE_ARRAY_COUNT(Texture); LayerIndex++)
+							{
+								Texture[LayerIndex][DstLinearIndex] = !bLoadingOutOfBounds ? TileUploadRequests[TileIndex].RenderState->TileStorage[SrcTileCoords].CPUTextureData[LayerIndex]->Data[SrcLinearIndex] : FLinearColor(0, 0, 0, 0);
+							}
 						}
 					}
 				});
 
-				RHICmdList.UnlockTexture2D(UploadTilePoolGPU->PooledRenderTargets[0]->GetRHI(), 0, false);
-				RHICmdList.UnlockTexture2D(UploadTilePoolGPU->PooledRenderTargets[1]->GetRHI(), 0, false);
-				RHICmdList.UnlockTexture2D(UploadTilePoolGPU->PooledRenderTargets[2]->GetRHI(), 0, false);
+				for (int32 LayerIndex = 0; LayerIndex < UE_ARRAY_COUNT(Texture); LayerIndex++)
+				{
+					RHICmdList.UnlockTexture2D(UploadTilePoolGPU->PooledRenderTargets[LayerIndex]->GetRHI(), 0, false);
+				}
 			}
 
 			FGPUBatchedTileRequests GPUBatchedTileRequests;
 			GPUBatchedTileRequests.BuildFromTileDescs(TileUploadRequests, LightmapTilePoolGPU, *ScratchTilePoolGPU);
 			GPUBatchedTileRequests.Commit();
 
-			IPooledRenderTarget* OutputRenderTargets[3] = { nullptr, nullptr, nullptr };
+			IPooledRenderTarget* OutputRenderTargets[4] = { nullptr, nullptr, nullptr, nullptr };
 
 			for (auto& Tile : TileUploadRequests)
 			{
-				for (int32 RenderTargetIndex = 0; RenderTargetIndex < 3; RenderTargetIndex++)
+				for (int32 RenderTargetIndex = 0; RenderTargetIndex < UE_ARRAY_COUNT(OutputRenderTargets); RenderTargetIndex++)
 				{
 					if (Tile.OutputRenderTargets[RenderTargetIndex] != nullptr)
 					{
@@ -1630,73 +1638,26 @@ void FLightmapRenderer::Finalize(FRDGBuilder& GraphBuilder)
 			DispatchResolution.X = GPreviewLightmapPhysicalTileSize * GPUBatchedTileRequests.BatchedTilesDesc.Num();
 			DispatchResolution.Y = GPreviewLightmapPhysicalTileSize;
 
-			FRDGTextureRef StagingHQLayer0 = GraphBuilder.RegisterExternalTexture(UploadTilePoolGPU->PooledRenderTargets[0], TEXT("StagingHQLayer0"));
-			FRDGTextureRef StagingHQLayer1 = GraphBuilder.RegisterExternalTexture(UploadTilePoolGPU->PooledRenderTargets[1], TEXT("StagingHQLayer1"));
-			FRDGTextureRef StagingShadowMask = GraphBuilder.RegisterExternalTexture(UploadTilePoolGPU->PooledRenderTargets[2], TEXT("StagingShadowMask"));
-
-			FBufferRHIRef SrcTilePositionsBuffer;
-			FShaderResourceViewRHIRef SrcTilePositionsSRV;
-			FBufferRHIRef DstTilePositionsBuffer;
-			FShaderResourceViewRHIRef DstTilePositionsSRV;
-
-			if (OutputRenderTargets[0] != nullptr)
+			FString StagingTextureNames[4] = { TEXT("StagingHQLayer0"), TEXT("StagingHQLayer1"), TEXT("StagingShadowMask"), TEXT("StagingSkyOcclusion") };
+			
+			for (int32 RenderTargetIndex = 0; RenderTargetIndex < UE_ARRAY_COUNT(OutputRenderTargets); RenderTargetIndex++)
 			{
+				if (OutputRenderTargets[RenderTargetIndex] != nullptr)
 				{
-					TResourceArray<FIntPoint> SrcTilePositions;
-					TResourceArray<FIntPoint> DstTilePositions;
-
-					for (int32 TileIndex = 0; TileIndex < TileUploadRequests.Num(); TileIndex++)
-					{
-						SrcTilePositions.Add(FIntPoint(TileIndex % UploadTilePoolGPU->SizeInTiles.X, TileIndex / UploadTilePoolGPU->SizeInTiles.X) * GPreviewLightmapPhysicalTileSize);
-						DstTilePositions.Add(TileUploadRequests[TileIndex].OutputPhysicalCoordinates[0] * GPreviewLightmapPhysicalTileSize);
-					}
-
-					{
-						FRHIResourceCreateInfo CreateInfo(TEXT("SrcTilePositionsBuffer"), &SrcTilePositions);
-						CreateInfo.GPUMask = FRHIGPUMask::GPU0();
-						SrcTilePositionsBuffer = RHICreateStructuredBuffer(sizeof(FIntPoint), SrcTilePositions.GetResourceDataSize(), BUF_Dynamic | BUF_ShaderResource, CreateInfo);
-						SrcTilePositionsSRV = RHICreateShaderResourceView(SrcTilePositionsBuffer);
-					}
-
-					{
-						FRHIResourceCreateInfo CreateInfo(TEXT("DstTilePositionsBuffer"), &DstTilePositions);
-						CreateInfo.GPUMask = FRHIGPUMask::GPU0();
-						DstTilePositionsBuffer = RHICreateStructuredBuffer(sizeof(FIntPoint), DstTilePositions.GetResourceDataSize(), BUF_Dynamic | BUF_ShaderResource, CreateInfo);
-						DstTilePositionsSRV = RHICreateShaderResourceView(DstTilePositionsBuffer);
-					}
-				}
-
-				{
-					FRDGTextureRef RenderTargetTileAtlas0 = GraphBuilder.RegisterExternalTexture(OutputRenderTargets[0], TEXT("GPULightmassRenderTargetTileAtlas0"));
+					FRDGTextureRef StagingTexture = GraphBuilder.RegisterExternalTexture(UploadTilePoolGPU->PooledRenderTargets[RenderTargetIndex], *StagingTextureNames[RenderTargetIndex]);
 
-					FUploadConvergedLightmapTilesCS::FParameters* PassParameters = GraphBuilder.AllocParameters<FUploadConvergedLightmapTilesCS::FParameters>();
+					FBufferRHIRef SrcTilePositionsBuffer;
+					FShaderResourceViewRHIRef SrcTilePositionsSRV;
+					FBufferRHIRef DstTilePositionsBuffer;
+					FShaderResourceViewRHIRef DstTilePositionsSRV;
 
-					PassParameters->NumBatchedTiles = GPUBatchedTileRequests.BatchedTilesDesc.Num();
-					PassParameters->SrcTexture = GraphBuilder.CreateUAV(StagingHQLayer0);
-					PassParameters->DstTexture = GraphBuilder.CreateUAV(RenderTargetTileAtlas0);
-					PassParameters->SrcTilePositions = HoldReference(GraphBuilder, SrcTilePositionsSRV);
-					PassParameters->DstTilePositions = HoldReference(GraphBuilder, DstTilePositionsSRV);
-
-					TShaderMapRef<FUploadConvergedLightmapTilesCS> ComputeShader(GlobalShaderMap);
-					FComputeShaderUtils::AddPass(
-						GraphBuilder,
-						RDG_EVENT_NAME("UploadConvergedLightmapTiles"),
-						ComputeShader,
-						PassParameters,
-						FComputeShaderUtils::GetGroupCount(DispatchResolution, FComputeShaderUtils::kGolden2DGroupSize));
-				}
-			}
-
-			if (OutputRenderTargets[1] != nullptr)
-			{
-				{
 					TResourceArray<FIntPoint> SrcTilePositions;
 					TResourceArray<FIntPoint> DstTilePositions;
-
+					
 					for (int32 TileIndex = 0; TileIndex < TileUploadRequests.Num(); TileIndex++)
 					{
 						SrcTilePositions.Add(FIntPoint(TileIndex % UploadTilePoolGPU->SizeInTiles.X, TileIndex / UploadTilePoolGPU->SizeInTiles.X) * GPreviewLightmapPhysicalTileSize);
-						DstTilePositions.Add(TileUploadRequests[TileIndex].OutputPhysicalCoordinates[1] * GPreviewLightmapPhysicalTileSize);
+						DstTilePositions.Add(TileUploadRequests[TileIndex].OutputPhysicalCoordinates[RenderTargetIndex] * GPreviewLightmapPhysicalTileSize);
 					}
 
 					{
@@ -1712,74 +1673,26 @@ void FLightmapRenderer::Finalize(FRDGBuilder& GraphBuilder)
 						DstTilePositionsBuffer = RHICreateStructuredBuffer(sizeof(FIntPoint), DstTilePositions.GetResourceDataSize(), BUF_Dynamic | BUF_ShaderResource, CreateInfo);
 						DstTilePositionsSRV = RHICreateShaderResourceView(DstTilePositionsBuffer);
 					}
-				}
 
-				{
-					FRDGTextureRef RenderTargetTileAtlas1 = GraphBuilder.RegisterExternalTexture(OutputRenderTargets[1], TEXT("GPULightmassRenderTargetTileAtlas1"));
-
-					FUploadConvergedLightmapTilesCS::FParameters* PassParameters = GraphBuilder.AllocParameters<FUploadConvergedLightmapTilesCS::FParameters>();
-
-					PassParameters->NumBatchedTiles = GPUBatchedTileRequests.BatchedTilesDesc.Num();
-					PassParameters->SrcTexture = GraphBuilder.CreateUAV(StagingHQLayer1);
-					PassParameters->DstTexture = GraphBuilder.CreateUAV(RenderTargetTileAtlas1);
-					PassParameters->SrcTilePositions = HoldReference(GraphBuilder, SrcTilePositionsSRV);
-					PassParameters->DstTilePositions = HoldReference(GraphBuilder, DstTilePositionsSRV);
-
-					TShaderMapRef<FUploadConvergedLightmapTilesCS> ComputeShader(GlobalShaderMap);
-					FComputeShaderUtils::AddPass(
-						GraphBuilder,
-						RDG_EVENT_NAME("UploadConvergedLightmapTiles"),
-						ComputeShader,
-						PassParameters,
-						FComputeShaderUtils::GetGroupCount(DispatchResolution, FComputeShaderUtils::kGolden2DGroupSize));
-				}
-			}
-
-			if (OutputRenderTargets[2] != nullptr)
-			{
-				{
-					TResourceArray<FIntPoint> SrcTilePositions;
-					TResourceArray<FIntPoint> DstTilePositions;
-
-					for (int32 TileIndex = 0; TileIndex < TileUploadRequests.Num(); TileIndex++)
 					{
-						SrcTilePositions.Add(FIntPoint(TileIndex % UploadTilePoolGPU->SizeInTiles.X, TileIndex / UploadTilePoolGPU->SizeInTiles.X) * GPreviewLightmapPhysicalTileSize);
-						DstTilePositions.Add(TileUploadRequests[TileIndex].OutputPhysicalCoordinates[2] * GPreviewLightmapPhysicalTileSize);
-					}
+						FRDGTextureRef RenderTargetTileAtlas = GraphBuilder.RegisterExternalTexture(OutputRenderTargets[0], *(FString(TEXT("GPULightmassRenderTargetTileAtlas_")) + StagingTextureNames[RenderTargetIndex]));
 
-					{
-						FRHIResourceCreateInfo CreateInfo(TEXT("SrcTilePositionsBuffer"), &SrcTilePositions);
-						CreateInfo.GPUMask = FRHIGPUMask::GPU0();
-						SrcTilePositionsBuffer = RHICreateStructuredBuffer(sizeof(FIntPoint), SrcTilePositions.GetResourceDataSize(), BUF_Dynamic | BUF_ShaderResource, CreateInfo);
-						SrcTilePositionsSRV = RHICreateShaderResourceView(SrcTilePositionsBuffer);
-					}
-
-					{
-						FRHIResourceCreateInfo CreateInfo(TEXT("DstTilePositionsBuffer"), &DstTilePositions);
-						CreateInfo.GPUMask = FRHIGPUMask::GPU0();
-						DstTilePositionsBuffer = RHICreateStructuredBuffer(sizeof(FIntPoint), DstTilePositions.GetResourceDataSize(), BUF_Dynamic | BUF_ShaderResource, CreateInfo);
-						DstTilePositionsSRV = RHICreateShaderResourceView(DstTilePositionsBuffer);
-					}
-				}
-
-				{
-					FRDGTextureRef RenderTargetTileAtlas2 = GraphBuilder.RegisterExternalTexture(OutputRenderTargets[2], TEXT("GPULightmassRenderTargetTileAtlas1"));
+						FUploadConvergedLightmapTilesCS::FParameters* PassParameters = GraphBuilder.AllocParameters<FUploadConvergedLightmapTilesCS::FParameters>();
 
-					FUploadConvergedLightmapTilesCS::FParameters* PassParameters = GraphBuilder.AllocParameters<FUploadConvergedLightmapTilesCS::FParameters>();
+						PassParameters->NumBatchedTiles = GPUBatchedTileRequests.BatchedTilesDesc.Num();
+						PassParameters->SrcTexture = GraphBuilder.CreateUAV(StagingTexture);
+						PassParameters->DstTexture = GraphBuilder.CreateUAV(RenderTargetTileAtlas);
+						PassParameters->SrcTilePositions = HoldReference(GraphBuilder, SrcTilePositionsSRV);
+						PassParameters->DstTilePositions = HoldReference(GraphBuilder, DstTilePositionsSRV);
 
-					PassParameters->NumBatchedTiles = GPUBatchedTileRequests.BatchedTilesDesc.Num();
-					PassParameters->SrcTexture = GraphBuilder.CreateUAV(StagingShadowMask);
-					PassParameters->DstTexture = GraphBuilder.CreateUAV(RenderTargetTileAtlas2);
-					PassParameters->SrcTilePositions = HoldReference(GraphBuilder, SrcTilePositionsSRV);
-					PassParameters->DstTilePositions = HoldReference(GraphBuilder, DstTilePositionsSRV);
-
-					TShaderMapRef<FUploadConvergedLightmapTilesCS> ComputeShader(GlobalShaderMap);
-					FComputeShaderUtils::AddPass(
-						GraphBuilder,
-						RDG_EVENT_NAME("UploadConvergedLightmapTiles"),
-						ComputeShader,
-						PassParameters,
-						FComputeShaderUtils::GetGroupCount(DispatchResolution, FComputeShaderUtils::kGolden2DGroupSize));
+						TShaderMapRef<FUploadConvergedLightmapTilesCS> ComputeShader(GlobalShaderMap);
+						FComputeShaderUtils::AddPass(
+							GraphBuilder,
+							RDG_EVENT_NAME("UploadConvergedLightmapTiles"),
+							ComputeShader,
+							PassParameters,
+							FComputeShaderUtils::GetGroupCount(DispatchResolution, FComputeShaderUtils::kGolden2DGroupSize));
+					}
 				}
 			}
 		}
@@ -2985,10 +2898,11 @@ void FLightmapRenderer::Finalize(FRDGBuilder& GraphBuilder)
 				LightmapReadbackGroup.ConvergedTileRequests = ConvergedTileRequestsOnCurrentGPU;
 				if (!LightmapReadbackGroup.ReadbackTilePoolGPU.IsValid())
 				{
-					LightmapReadbackGroup.ReadbackTilePoolGPU = MakeUnique<FLightmapTilePoolGPU>(3, FIntPoint(NewSize, NewSize), FIntPoint(GPreviewLightmapPhysicalTileSize, GPreviewLightmapPhysicalTileSize));
+					LightmapReadbackGroup.ReadbackTilePoolGPU = MakeUnique<FLightmapTilePoolGPU>(4, FIntPoint(NewSize, NewSize), FIntPoint(GPreviewLightmapPhysicalTileSize, GPreviewLightmapPhysicalTileSize));
 					LightmapReadbackGroup.StagingHQLayer0Readback = MakeUnique<FRHIGPUTextureReadback>(TEXT("StagingHQLayer0Readback"));
 					LightmapReadbackGroup.StagingHQLayer1Readback = MakeUnique<FRHIGPUTextureReadback>(TEXT("StagingHQLayer1Readback"));
 					LightmapReadbackGroup.StagingShadowMaskReadback = MakeUnique<FRHIGPUTextureReadback>(TEXT("StagingShadowMaskReadback"));
+					LightmapReadbackGroup.StagingSkyOcclusionReadback = MakeUnique<FRHIGPUTextureReadback>(TEXT("StagingSkyOcclusionReadback"));
 				}
 
 				FGPUBatchedTileRequests GPUBatchedTileRequests;
@@ -3010,6 +2924,7 @@ void FLightmapRenderer::Finalize(FRDGBuilder& GraphBuilder)
 				FRDGTextureRef StagingHQLayer0 = GraphBuilder.RegisterExternalTexture(LightmapReadbackGroup.ReadbackTilePoolGPU->PooledRenderTargets[0], TEXT("StagingHQLayer0"));
 				FRDGTextureRef StagingHQLayer1 = GraphBuilder.RegisterExternalTexture(LightmapReadbackGroup.ReadbackTilePoolGPU->PooledRenderTargets[1], TEXT("StagingHQLayer1"));
 				FRDGTextureRef StagingShadowMask = GraphBuilder.RegisterExternalTexture(LightmapReadbackGroup.ReadbackTilePoolGPU->PooledRenderTargets[2], TEXT("StagingShadowMask"));
+				FRDGTextureRef StagingSkyOcclusion = GraphBuilder.RegisterExternalTexture(LightmapReadbackGroup.ReadbackTilePoolGPU->PooledRenderTargets[3], TEXT("StagingSkyOcclusion"));
 
 				{
 					FCopyConvergedLightmapTilesCS::FParameters* PassParameters = GraphBuilder.AllocParameters<FCopyConvergedLightmapTilesCS::FParameters>();
@@ -3025,6 +2940,7 @@ void FLightmapRenderer::Finalize(FRDGBuilder& GraphBuilder)
 					PassParameters->StagingHQLayer0 = GraphBuilder.CreateUAV(StagingHQLayer0);
 					PassParameters->StagingHQLayer1 = GraphBuilder.CreateUAV(StagingHQLayer1);
 					PassParameters->StagingShadowMask = GraphBuilder.CreateUAV(StagingShadowMask);
+					PassParameters->StagingSkyOcclusion = GraphBuilder.CreateUAV(StagingSkyOcclusion);
 
 					TShaderMapRef<FCopyConvergedLightmapTilesCS> ComputeShader(GlobalShaderMap);
 					FComputeShaderUtils::AddPass(
@@ -3038,6 +2954,7 @@ void FLightmapRenderer::Finalize(FRDGBuilder& GraphBuilder)
 				AddEnqueueCopyPass(GraphBuilder, LightmapReadbackGroup.StagingHQLayer0Readback.Get(), StagingHQLayer0);
 				AddEnqueueCopyPass(GraphBuilder, LightmapReadbackGroup.StagingHQLayer1Readback.Get(), StagingHQLayer1);
 				AddEnqueueCopyPass(GraphBuilder, LightmapReadbackGroup.StagingShadowMaskReadback.Get(), StagingShadowMask);
+				AddEnqueueCopyPass(GraphBuilder, LightmapReadbackGroup.StagingSkyOcclusionReadback.Get(), StagingSkyOcclusion);
 
 				OngoingReadbacks.Emplace(ReadbackGroupToUse);
 			}
@@ -3161,7 +3078,8 @@ void FLightmapRenderer::BackgroundTick()
 
 		if (ReadbackGroup.StagingHQLayer0Readback->IsReady(FRHIGPUMask::FromIndex(ReadbackGroup.GPUIndex)) &&
 			ReadbackGroup.StagingHQLayer1Readback->IsReady(FRHIGPUMask::FromIndex(ReadbackGroup.GPUIndex)) &&
-			ReadbackGroup.StagingShadowMaskReadback->IsReady(FRHIGPUMask::FromIndex(ReadbackGroup.GPUIndex)))
+			ReadbackGroup.StagingShadowMaskReadback->IsReady(FRHIGPUMask::FromIndex(ReadbackGroup.GPUIndex)) &&
+			ReadbackGroup.StagingSkyOcclusionReadback->IsReady(FRHIGPUMask::FromIndex(ReadbackGroup.GPUIndex)))
 		{
 			FRHICommandListImmediate& RHICmdList = FRHICommandListExecutor::GetImmediateCommandList();
 			SCOPED_GPU_MASK(RHICmdList, FRHIGPUMask::FromIndex(ReadbackGroup.GPUIndex));
@@ -3172,21 +3090,25 @@ void FLightmapRenderer::BackgroundTick()
 
 			// FLinearColor is in RGBA while the GPU texture is in ABGR
 			// TODO: apply swizzling in the copy compute shader if this becomes a problem
-			void* LockedData[3];
+			void* LockedData[4];
 			LockedData[0] = ReadbackGroup.StagingHQLayer0Readback->Lock(ReadbackGroup.TextureData->RowPitchInPixels[0]); // This forces a GPU stall
 			LockedData[1] = ReadbackGroup.StagingHQLayer1Readback->Lock(ReadbackGroup.TextureData->RowPitchInPixels[1]); // This forces a GPU stall
 			LockedData[2] = ReadbackGroup.StagingShadowMaskReadback->Lock(ReadbackGroup.TextureData->RowPitchInPixels[2]); // This forces a GPU stall
+			LockedData[3] = ReadbackGroup.StagingSkyOcclusionReadback->Lock(ReadbackGroup.TextureData->RowPitchInPixels[3]); // This forces a GPU stall
 
 			ReadbackGroup.TextureData->Texture[0].AddUninitialized(ReadbackGroup.ReadbackTilePoolGPU->SizeInTiles.Y * GPreviewLightmapPhysicalTileSize * ReadbackGroup.TextureData->RowPitchInPixels[0]);
 			ReadbackGroup.TextureData->Texture[1].AddUninitialized(ReadbackGroup.ReadbackTilePoolGPU->SizeInTiles.Y * GPreviewLightmapPhysicalTileSize * ReadbackGroup.TextureData->RowPitchInPixels[1]);
 			ReadbackGroup.TextureData->Texture[2].AddUninitialized(ReadbackGroup.ReadbackTilePoolGPU->SizeInTiles.Y * GPreviewLightmapPhysicalTileSize * ReadbackGroup.TextureData->RowPitchInPixels[2]);
+			ReadbackGroup.TextureData->Texture[3].AddUninitialized(ReadbackGroup.ReadbackTilePoolGPU->SizeInTiles.Y * GPreviewLightmapPhysicalTileSize * ReadbackGroup.TextureData->RowPitchInPixels[3]);
 			FMemory::Memcpy(ReadbackGroup.TextureData->Texture[0].GetData(), LockedData[0], ReadbackGroup.ReadbackTilePoolGPU->SizeInTiles.Y * GPreviewLightmapPhysicalTileSize * ReadbackGroup.TextureData->RowPitchInPixels[0] * sizeof(FLinearColor));
 			FMemory::Memcpy(ReadbackGroup.TextureData->Texture[1].GetData(), LockedData[1], ReadbackGroup.ReadbackTilePoolGPU->SizeInTiles.Y * GPreviewLightmapPhysicalTileSize * ReadbackGroup.TextureData->RowPitchInPixels[1] * sizeof(FLinearColor));
 			FMemory::Memcpy(ReadbackGroup.TextureData->Texture[2].GetData(), LockedData[2], ReadbackGroup.ReadbackTilePoolGPU->SizeInTiles.Y * GPreviewLightmapPhysicalTileSize * ReadbackGroup.TextureData->RowPitchInPixels[2] * sizeof(FLinearColor));
+			FMemory::Memcpy(ReadbackGroup.TextureData->Texture[3].GetData(), LockedData[3], ReadbackGroup.ReadbackTilePoolGPU->SizeInTiles.Y * GPreviewLightmapPhysicalTileSize * ReadbackGroup.TextureData->RowPitchInPixels[3] * sizeof(FLinearColor));
 
 			ReadbackGroup.StagingHQLayer0Readback->Unlock();
 			ReadbackGroup.StagingHQLayer1Readback->Unlock();
 			ReadbackGroup.StagingShadowMaskReadback->Unlock();
+			ReadbackGroup.StagingSkyOcclusionReadback->Unlock();
 
 			for (int32 TileIndex = 0; TileIndex < ReadbackGroup.ConvergedTileRequests.Num(); TileIndex++)
 			{
@@ -3213,6 +3135,7 @@ void FLightmapRenderer::BackgroundTick()
 				TileStorage.CPUTextureData[0]->Decompress();
 				TileStorage.CPUTextureData[1]->Decompress();
 				TileStorage.CPUTextureData[2]->Decompress();
+				TileStorage.CPUTextureData[3]->Decompress();
 
 				for (int32 Y = 0; Y < GPreviewLightmapVirtualTileSize; Y++)
 				{
@@ -3236,6 +3159,7 @@ void FLightmapRenderer::BackgroundTick()
 
 						// For shadow maps, pass through
 						TileStorage.CPUTextureData[2]->Data[DstLinearIndex] = ReadbackGroup.TextureData->Texture[2][SrcLinearIndex];
+						TileStorage.CPUTextureData[3]->Data[DstLinearIndex] = ReadbackGroup.TextureData->Texture[3][SrcLinearIndex];
 					}
 				}
 			}
@@ -3463,14 +3387,14 @@ void FLightmapRenderer::BackgroundTick()
 							{
 								bAnyTileSelected = true;
 
-								FVTProduceTargetLayer TargetLayers[3];
+								FVTProduceTargetLayer TargetLayers[4];
 
 								Lightmap.LightmapPreviewVirtualTexture->ProducePageData(
 									RHICmdList,
 									Scene->FeatureLevel,
 									EVTProducePageFlags::None,
 									FVirtualTextureProducerHandle(),
-									0b111,
+									0b1111,
 									0,
 									FMath::MortonCode2(X) | (FMath::MortonCode2(Y) << 1),
 									0,
diff --git a/Engine/Plugins/Experimental/GPULightmass/Source/GPULightmass/Private/LightmapRenderer.h b/Engine/Plugins/Experimental/GPULightmass/Source/GPULightmass/Private/LightmapRenderer.h
--- a/Engine/Plugins/Experimental/GPULightmass/Source/GPULightmass/Private/LightmapRenderer.h
+++ b/Engine/Plugins/Experimental/GPULightmass/Source/GPULightmass/Private/LightmapRenderer.h
@@ -60,13 +60,14 @@ struct FLightmapReadbackGroup
 	TUniquePtr<FRHIGPUTextureReadback> StagingHQLayer0Readback;
 	TUniquePtr<FRHIGPUTextureReadback> StagingHQLayer1Readback;
 	TUniquePtr<FRHIGPUTextureReadback> StagingShadowMaskReadback;
+	TUniquePtr<FRHIGPUTextureReadback> StagingSkyOcclusionReadback;
 
 	struct FTextureData
 	{
 		// Duplicate some metadata so the async thread doesn't need to access FLightmapReadbackGroup
 		FIntPoint SizeInTiles;
-		int32 RowPitchInPixels[3];
-		TArray<FLinearColor> Texture[3];
+		int32 RowPitchInPixels[4];
+		TArray<FLinearColor> Texture[4];
 
 		volatile int32 bDenoisingFinished = 0;
 	};
diff --git a/Engine/Plugins/Experimental/GPULightmass/Source/GPULightmass/Private/LightmapStorage.cpp b/Engine/Plugins/Experimental/GPULightmass/Source/GPULightmass/Private/LightmapStorage.cpp
--- a/Engine/Plugins/Experimental/GPULightmass/Source/GPULightmass/Private/LightmapStorage.cpp
+++ b/Engine/Plugins/Experimental/GPULightmass/Source/GPULightmass/Private/LightmapStorage.cpp
@@ -125,6 +125,7 @@ void FLightmap::CreateGameThreadResources()
 		TextureUObject->SetLayerForType(ELightMapVirtualTextureType::LightmapLayer0, 0);
 		TextureUObject->SetLayerForType(ELightMapVirtualTextureType::LightmapLayer1, 1);
 		TextureUObject->SetLayerForType(ELightMapVirtualTextureType::ShadowMask, 2);
+		TextureUObject->SetLayerForType(ELightMapVirtualTextureType::SkyOcclusion, 3);
 	}
 
 	TextureUObjectGuard = MakeUnique<FGCObjectScopeGuard>(TextureUObject);
diff --git a/Engine/Plugins/Experimental/GPULightmass/Source/GPULightmass/Private/Scene/Scene.cpp b/Engine/Plugins/Experimental/GPULightmass/Source/GPULightmass/Private/Scene/Scene.cpp
--- a/Engine/Plugins/Experimental/GPULightmass/Source/GPULightmass/Private/Scene/Scene.cpp
+++ b/Engine/Plugins/Experimental/GPULightmass/Source/GPULightmass/Private/Scene/Scene.cpp
@@ -599,7 +599,7 @@ void FScene::AddLight(USkyLightComponent* SkyLight)
 	LightScene.SkyLight = MoveTemp(NewSkyLight);
 
 	FSkyLightRenderState NewSkyLightRenderState;
-	NewSkyLightRenderState.bStationary = SkyLight->CastShadows && SkyLight->CastStaticShadows && !SkyLight->HasStaticLighting();
+	NewSkyLightRenderState.bStationary = SkyLight->Mobility == EComponentMobility::Stationary;
 	NewSkyLightRenderState.bCastShadow = SkyLight->CastShadows && SkyLight->CastStaticShadows;
 	NewSkyLightRenderState.Color = SkyLight->GetLightColor() * SkyLight->Intensity;
 	NewSkyLightRenderState.TextureDimensions = FIntPoint(SkyLight->GetProcessedSkyTexture()->GetSizeX(), SkyLight->GetProcessedSkyTexture()->GetSizeY());
