GPULM: Implement transcoding for stationary skylight shadows. Results can now be saved. No denoising support yet

diff --git a/Engine/Plugins/Experimental/GPULightmass/Source/GPULightmass/Private/LightmapEncoding.cpp b/Engine/Plugins/Experimental/GPULightmass/Source/GPULightmass/Private/LightmapEncoding.cpp
--- a/Engine/Plugins/Experimental/GPULightmass/Source/GPULightmass/Private/LightmapEncoding.cpp
+++ b/Engine/Plugins/Experimental/GPULightmass/Source/GPULightmass/Private/LightmapEncoding.cpp
@@ -347,6 +347,16 @@ FLightSampleData ConvertToLightSample(FLinearColor IncidentLighting, FLinearColo
 	return Sample;
 }
 
+void AddSkyOcclusionToLightSample(FLightSampleData& Sample, FLinearColor SkyOcclusion)
+{
+	// Revert sqrt in LightmapEncoding.ush for preview
+	float Length = SkyOcclusion.A * SkyOcclusion.A;
+	FVector3f UnpackedBentNormalVector = FVector3f(SkyOcclusion) * 2.0f - 1.0f;
+	Sample.SkyOcclusion[0] = UnpackedBentNormalVector.X * Length;
+	Sample.SkyOcclusion[1] = UnpackedBentNormalVector.Y * Length;
+	Sample.SkyOcclusion[2] = UnpackedBentNormalVector.Z * Length;
+}
+
 FQuantizedSignedDistanceFieldShadowSample ConvertToShadowSample(FLinearColor ShadowMask, int32 ChannelIndex)
 {
 	FQuantizedSignedDistanceFieldShadowSample Sample;
diff --git a/Engine/Plugins/Experimental/GPULightmass/Source/GPULightmass/Private/LightmapEncoding.h b/Engine/Plugins/Experimental/GPULightmass/Source/GPULightmass/Private/LightmapEncoding.h
--- a/Engine/Plugins/Experimental/GPULightmass/Source/GPULightmass/Private/LightmapEncoding.h
+++ b/Engine/Plugins/Experimental/GPULightmass/Source/GPULightmass/Private/LightmapEncoding.h
@@ -51,6 +51,7 @@ struct FLightSampleData
 };
 
 FLightSampleData ConvertToLightSample(FLinearColor IncidentLighting, FLinearColor LuminanceSH);
+void AddSkyOcclusionToLightSample(FLightSampleData& Sample, FLinearColor SkyOcclusion);
 
 void QuantizeLightSamples(
 	TArray<FLightSampleData> InLightSamples,
diff --git a/Engine/Plugins/Experimental/GPULightmass/Source/GPULightmass/Private/Scene/Lights.h b/Engine/Plugins/Experimental/GPULightmass/Source/GPULightmass/Private/Scene/Lights.h
--- a/Engine/Plugins/Experimental/GPULightmass/Source/GPULightmass/Private/Scene/Lights.h
+++ b/Engine/Plugins/Experimental/GPULightmass/Source/GPULightmass/Private/Scene/Lights.h
@@ -287,7 +287,9 @@ using FRectLightRenderStateRef = TEntityArray<FRectLightRenderState>::EntityRefT
 struct FSkyLightRenderState
 {
 	bool bStationary = false;
-	bool bCastShadow = true;
+	bool bCastShadow = true;	
+	bool CastsStationaryShadow() { return bStationary && bCastShadow; }
+	
 	FLinearColor Color;
 	FTextureRHIRef ProcessedTexture;
 	FSamplerStateRHIRef ProcessedTextureSampler;
@@ -306,6 +308,10 @@ struct FSkyLightRenderState
 struct FSkyLightBuildInfo
 {
 	USkyLightComponent* ComponentUObject = nullptr;
+
+	bool bStationary = false;
+	bool bCastShadow = true;	
+	bool CastsStationaryShadow() { return bStationary && bCastShadow; }
 };
 
 class FLightArrayBase
diff --git a/Engine/Plugins/Experimental/GPULightmass/Source/GPULightmass/Private/Scene/Scene.cpp b/Engine/Plugins/Experimental/GPULightmass/Source/GPULightmass/Private/Scene/Scene.cpp
--- a/Engine/Plugins/Experimental/GPULightmass/Source/GPULightmass/Private/Scene/Scene.cpp
+++ b/Engine/Plugins/Experimental/GPULightmass/Source/GPULightmass/Private/Scene/Scene.cpp
@@ -595,6 +595,8 @@ void FScene::AddLight(USkyLightComponent* SkyLight)
 
 	FSkyLightBuildInfo NewSkyLight;
 	NewSkyLight.ComponentUObject = SkyLight;
+	NewSkyLight.bStationary = SkyLight->Mobility == EComponentMobility::Stationary;
+	NewSkyLight.bCastShadow = SkyLight->CastShadows && SkyLight->CastStaticShadows;
 
 	LightScene.SkyLight = MoveTemp(NewSkyLight);
 
@@ -1851,6 +1853,55 @@ void FLocalLightBuildInfo::AllocateMapBuildData(ULevel* StorageLevel)
 	LightBuildData.ShadowMapChannel = CastsStationaryShadow() ? GetComponentUObject()->PreviewShadowMapChannel : INDEX_NONE;
 }
 
+void FScene::AddRelevantStaticLightGUIDs(FQuantizedLightmapData* QuantizedLightmapData, const FBoxSphereBounds& WorldBounds)
+{
+	// Add static lights to lightmap data
+	for (FDirectionalLightBuildInfo& DirectionalLight : LightScene.DirectionalLights.Elements)
+	{
+		if (!DirectionalLight.bStationary)
+		{
+			UDirectionalLightComponent* Light = DirectionalLight.ComponentUObject;
+			QuantizedLightmapData->LightGuids.Add(Light->LightGuid);
+		}
+	}
+
+	for (FPointLightBuildInfo& PointLight : LightScene.PointLights.Elements)
+	{
+		if (!PointLight.bStationary)
+		{
+			UPointLightComponent* Light = PointLight.ComponentUObject;
+			if (PointLight.AffectsBounds(WorldBounds))
+			{
+				QuantizedLightmapData->LightGuids.Add(Light->LightGuid);
+			}
+		}
+	}
+
+	for (FSpotLightBuildInfo& SpotLight : LightScene.SpotLights.Elements)
+	{
+		if (!SpotLight.bStationary)
+		{
+			USpotLightComponent* Light = SpotLight.ComponentUObject;
+			if (SpotLight.AffectsBounds(WorldBounds))
+			{
+				QuantizedLightmapData->LightGuids.Add(Light->LightGuid);
+			}
+		}
+	}
+
+	for (FRectLightBuildInfo& RectLight : LightScene.RectLights.Elements)
+	{
+		if (!RectLight.bStationary)
+		{
+			URectLightComponent* Light = RectLight.ComponentUObject;
+			if (RectLight.AffectsBounds(WorldBounds))
+			{
+				QuantizedLightmapData->LightGuids.Add(Light->LightGuid);
+			}
+		}
+	}
+}
+
 void FScene::ApplyFinishedLightmapsToWorld()
 {
 	UWorld* World = GPULightmass->World;
@@ -1858,6 +1909,8 @@ void FScene::ApplyFinishedLightmapsToWorld()
 	static const auto CVar = IConsoleManager::Get().FindTConsoleVariableDataInt(TEXT("r.VirtualTexturedLightmaps"));
 	const bool bUseVirtualTextures = (CVar->GetValueOnAnyThread() != 0) && UseVirtualTexturing(World->FeatureLevel);
 
+	bool bHasSkyShadowing = LightScene.SkyLight.IsSet() && LightScene.SkyLight->CastsStationaryShadow();
+	
 	{
 		FScopedSlowTask SlowTask(3);
 		SlowTask.MakeDialog();
@@ -2091,6 +2144,19 @@ void FScene::ApplyFinishedLightmapsToWorld()
 							{
 								LightSampleData[DstLinearIndex] = ConvertToLightSample(Lightmap.TileStorage[FTileVirtualCoordinates(SrcTilePosition, 0)].CPUTextureData[0]->Data[SrcLinearIndex], Lightmap.TileStorage[FTileVirtualCoordinates(SrcTilePosition, 0)].CPUTextureData[1]->Data[SrcLinearIndex]);
 							});
+
+							if (bHasSkyShadowing)
+							{
+								CopyRectTiled(
+									FIntPoint(0, 0),
+									FIntRect(FIntPoint(0, 0), Lightmap.GetSize()),
+									SrcRowPitchInPixels,
+									DstRowPitchInPixels,
+									[&Lightmap, &LightSampleData](int32 DstLinearIndex, FIntPoint SrcTilePosition, int32 SrcLinearIndex) mutable
+								{
+									AddSkyOcclusionToLightSample(LightSampleData[DstLinearIndex], Lightmap.TileStorage[FTileVirtualCoordinates(SrcTilePosition, 0)].CPUTextureData[3]->Data[SrcLinearIndex]);
+								});
+							}
 						}
 
 #if 0
@@ -2149,56 +2215,11 @@ void FScene::ApplyFinishedLightmapsToWorld()
 						FQuantizedLightmapData* QuantizedLightmapData = new FQuantizedLightmapData();
 						QuantizedLightmapData->SizeX = Lightmap.GetSize().X;
 						QuantizedLightmapData->SizeY = Lightmap.GetSize().Y;
+						QuantizedLightmapData->bHasSkyShadowing = bHasSkyShadowing;
 
 						QuantizeLightSamples(LightSampleData, QuantizedLightmapData->Data, QuantizedLightmapData->Scale, QuantizedLightmapData->Add);
 
-						// Add static lights to lightmap data
-						{
-							for (FDirectionalLightBuildInfo& DirectionalLight : LightScene.DirectionalLights.Elements)
-							{
-								if (!DirectionalLight.bStationary)
-								{
-									UDirectionalLightComponent* Light = DirectionalLight.ComponentUObject;
-									QuantizedLightmapData->LightGuids.Add(Light->LightGuid);
-								}
-							}
-
-							for (FPointLightBuildInfo& PointLight : LightScene.PointLights.Elements)
-							{
-								if (!PointLight.bStationary)
-								{
-									UPointLightComponent* Light = PointLight.ComponentUObject;
-									if (PointLight.AffectsBounds(StaticMeshInstances.Elements[InstanceIndex].WorldBounds))
-									{
-										QuantizedLightmapData->LightGuids.Add(Light->LightGuid);
-									}
-								}
-							}
-
-							for (FSpotLightBuildInfo& SpotLight : LightScene.SpotLights.Elements)
-							{
-								if (!SpotLight.bStationary)
-								{
-									USpotLightComponent* Light = SpotLight.ComponentUObject;
-									if (SpotLight.AffectsBounds(StaticMeshInstances.Elements[InstanceIndex].WorldBounds))
-									{
-										QuantizedLightmapData->LightGuids.Add(Light->LightGuid);
-									}
-								}
-							}
-
-							for (FRectLightBuildInfo& RectLight : LightScene.RectLights.Elements)
-							{
-								if (!RectLight.bStationary)
-								{
-									URectLightComponent* Light = RectLight.ComponentUObject;
-									if (RectLight.AffectsBounds(StaticMeshInstances.Elements[InstanceIndex].WorldBounds))
-									{
-										QuantizedLightmapData->LightGuids.Add(Light->LightGuid);
-									}
-								}
-							}
-						}
+						AddRelevantStaticLightGUIDs(QuantizedLightmapData, StaticMeshInstances.Elements[InstanceIndex].WorldBounds);
 
 						// Transencode stationary light shadow masks
 						TMap<ULightComponent*, FShadowMapData2D*> ShadowMaps;
@@ -2306,7 +2327,7 @@ void FScene::ApplyFinishedLightmapsToWorld()
 
 								FMeshMapBuildData& MeshBuildData = Registry->AllocateMeshBuildData(ComponentLODInfo.MapBuildDataId, true);
 
-								const bool bNeedsLightMap = true;// bHasNonZeroData;
+								const bool bNeedsLightMap = bHasNonZeroData || (bUseVirtualTextures && ShadowMaps.Num() > 0);
 								if (bNeedsLightMap)
 								{
 									// Create a light-map for the primitive.
@@ -2412,6 +2433,19 @@ void FScene::ApplyFinishedLightmapsToWorld()
 								{
 									LightSampleData[DstLinearIndex] = ConvertToLightSample(Lightmap.TileStorage[FTileVirtualCoordinates(SrcTilePosition, 0)].CPUTextureData[0]->Data[SrcLinearIndex], Lightmap.TileStorage[FTileVirtualCoordinates(SrcTilePosition, 0)].CPUTextureData[1]->Data[SrcLinearIndex]);
 								});
+
+								if (bHasSkyShadowing)
+								{
+									CopyRectTiled(
+										InstanceTileMin,
+										FIntRect(FIntPoint(0, 0), FIntPoint(BaseLightMapWidth, BaseLightMapHeight)),
+										SrcRowPitchInPixels,
+										DstRowPitchInPixels,
+										[&Lightmap, &LightSampleData](int32 DstLinearIndex, FIntPoint SrcTilePosition, int32 SrcLinearIndex) mutable
+									{
+										AddSkyOcclusionToLightSample(LightSampleData[DstLinearIndex], Lightmap.TileStorage[FTileVirtualCoordinates(SrcTilePosition, 0)].CPUTextureData[3]->Data[SrcLinearIndex]);
+									});
+								}
 							}
 
 							if (Settings->DenoisingOptions == EGPULightmassDenoisingOptions::OnCompletion)
@@ -2422,6 +2456,7 @@ void FScene::ApplyFinishedLightmapsToWorld()
 							FQuantizedLightmapData& QuantizedLightmapData = *InstancedSourceQuantizedData[InstanceIndex];
 							QuantizedLightmapData.SizeX = BaseLightMapWidth;
 							QuantizedLightmapData.SizeY = BaseLightMapHeight;
+							QuantizedLightmapData.bHasSkyShadowing = bHasSkyShadowing;
 
 							QuantizeLightSamples(LightSampleData, QuantizedLightmapData.Data, QuantizedLightmapData.Scale, QuantizedLightmapData.Add);
 
@@ -2542,52 +2577,8 @@ void FScene::ApplyFinishedLightmapsToWorld()
 						if (InstancedSourceQuantizedData.Num() > 0)
 						{
 							TUniquePtr<FQuantizedLightmapData>& QuantizedLightmapData = InstancedSourceQuantizedData[0];
-							{
-								for (FDirectionalLightBuildInfo& DirectionalLight : LightScene.DirectionalLights.Elements)
-								{
-									if (!DirectionalLight.bStationary)
-									{
-										UDirectionalLightComponent* Light = DirectionalLight.ComponentUObject;
-										QuantizedLightmapData->LightGuids.Add(Light->LightGuid);
-									}
-								}
-
-								for (FPointLightBuildInfo& PointLight : LightScene.PointLights.Elements)
-								{
-									if (!PointLight.bStationary)
-									{
-										UPointLightComponent* Light = PointLight.ComponentUObject;
-										if (PointLight.AffectsBounds(InstanceGroup.WorldBounds))
-										{
-											QuantizedLightmapData->LightGuids.Add(Light->LightGuid);
-										}
-									}
-								}
-
-								for (FSpotLightBuildInfo& SpotLight : LightScene.SpotLights.Elements)
-								{
-									if (!SpotLight.bStationary)
-									{
-										USpotLightComponent* Light = SpotLight.ComponentUObject;
-										if (SpotLight.AffectsBounds(InstanceGroup.WorldBounds))
-										{
-											QuantizedLightmapData->LightGuids.Add(Light->LightGuid);
-										}
-									}
-								}
-
-								for (FRectLightBuildInfo& RectLight : LightScene.RectLights.Elements)
-								{
-									if (!RectLight.bStationary)
-									{
-										URectLightComponent* Light = RectLight.ComponentUObject;
-										if (RectLight.AffectsBounds(InstanceGroup.WorldBounds))
-										{
-											QuantizedLightmapData->LightGuids.Add(Light->LightGuid);
-										}
-									}
-								}
-							}
+							
+							AddRelevantStaticLightGUIDs(QuantizedLightmapData.Get(), InstanceGroup.WorldBounds);
 						}
 
 						UStaticMesh* ResolvedMesh = InstanceGroup.ComponentUObject->GetStaticMesh();
@@ -2702,6 +2693,19 @@ void FScene::ApplyFinishedLightmapsToWorld()
 							{
 								LightSampleData[DstLinearIndex] = ConvertToLightSample(Lightmap.TileStorage[FTileVirtualCoordinates(SrcTilePosition, 0)].CPUTextureData[0]->Data[SrcLinearIndex], Lightmap.TileStorage[FTileVirtualCoordinates(SrcTilePosition, 0)].CPUTextureData[1]->Data[SrcLinearIndex]);
 							});
+
+							if (bHasSkyShadowing)
+							{
+								CopyRectTiled(
+									FIntPoint(0, 0),
+									FIntRect(FIntPoint(0, 0), Lightmap.GetSize()),
+									SrcRowPitchInPixels,
+									DstRowPitchInPixels,
+									[&Lightmap, &LightSampleData](int32 DstLinearIndex, FIntPoint SrcTilePosition, int32 SrcLinearIndex) mutable
+								{
+									AddSkyOcclusionToLightSample(LightSampleData[DstLinearIndex], Lightmap.TileStorage[FTileVirtualCoordinates(SrcTilePosition, 0)].CPUTextureData[3]->Data[SrcLinearIndex]);
+								});
+							}
 						}
 
 						if (Settings->DenoisingOptions == EGPULightmassDenoisingOptions::OnCompletion)
@@ -2712,56 +2716,11 @@ void FScene::ApplyFinishedLightmapsToWorld()
 						FQuantizedLightmapData* QuantizedLightmapData = new FQuantizedLightmapData();
 						QuantizedLightmapData->SizeX = Lightmap.GetSize().X;
 						QuantizedLightmapData->SizeY = Lightmap.GetSize().Y;
+						QuantizedLightmapData->bHasSkyShadowing = bHasSkyShadowing;
 
 						QuantizeLightSamples(LightSampleData, QuantizedLightmapData->Data, QuantizedLightmapData->Scale, QuantizedLightmapData->Add);
 
-						// Add static lights to lightmap data
-						{
-							for (FDirectionalLightBuildInfo& DirectionalLight : LightScene.DirectionalLights.Elements)
-							{
-								if (!DirectionalLight.bStationary)
-								{
-									UDirectionalLightComponent* Light = DirectionalLight.ComponentUObject;
-									QuantizedLightmapData->LightGuids.Add(Light->LightGuid);
-								}
-							}
-
-							for (FPointLightBuildInfo& PointLight : LightScene.PointLights.Elements)
-							{
-								if (!PointLight.bStationary)
-								{
-									UPointLightComponent* Light = PointLight.ComponentUObject;
-									if (PointLight.AffectsBounds(Landscapes.Elements[LandscapeIndex].WorldBounds))
-									{
-										QuantizedLightmapData->LightGuids.Add(Light->LightGuid);
-									}
-								}
-							}
-
-							for (FSpotLightBuildInfo& SpotLight : LightScene.SpotLights.Elements)
-							{
-								if (!SpotLight.bStationary)
-								{
-									USpotLightComponent* Light = SpotLight.ComponentUObject;
-									if (SpotLight.AffectsBounds(Landscapes.Elements[LandscapeIndex].WorldBounds))
-									{
-										QuantizedLightmapData->LightGuids.Add(Light->LightGuid);
-									}
-								}
-							}
-
-							for (FRectLightBuildInfo& RectLight : LightScene.RectLights.Elements)
-							{
-								if (!RectLight.bStationary)
-								{
-									URectLightComponent* Light = RectLight.ComponentUObject;
-									if (RectLight.AffectsBounds(Landscapes.Elements[LandscapeIndex].WorldBounds))
-									{
-										QuantizedLightmapData->LightGuids.Add(Light->LightGuid);
-									}
-								}
-							}
-						}
+						AddRelevantStaticLightGUIDs(QuantizedLightmapData, Landscapes.Elements[LandscapeIndex].WorldBounds);
 
 						// Transencode stationary light shadow masks
 						TMap<ULightComponent*, FShadowMapData2D*> ShadowMaps;
@@ -2826,7 +2785,7 @@ void FScene::ApplyFinishedLightmapsToWorld()
 						{
 							ULandscapeComponent* LandscapeComponent = Landscapes.Elements[LandscapeIndex].ComponentUObject;
 							ELightMapPaddingType PaddingType = LMPT_NoPadding;
-							const bool bHasNonZeroData = QuantizedLightmapData->HasNonZeroData();
+							const bool bHasNonZeroData = QuantizedLightmapData->HasNonZeroData() || (bUseVirtualTextures && ShadowMaps.Num() > 0);
 
 							ULevel* StorageLevel = LightingScenario ? LightingScenario : LandscapeComponent->GetOwner()->GetLevel();
 							UMapBuildDataRegistry* Registry = StorageLevel->GetOrCreateMapBuildData();
diff --git a/Engine/Plugins/Experimental/GPULightmass/Source/GPULightmass/Private/Scene/Scene.h b/Engine/Plugins/Experimental/GPULightmass/Source/GPULightmass/Private/Scene/Scene.h
--- a/Engine/Plugins/Experimental/GPULightmass/Source/GPULightmass/Private/Scene/Scene.h
+++ b/Engine/Plugins/Experimental/GPULightmass/Source/GPULightmass/Private/Scene/Scene.h
@@ -178,6 +178,7 @@ public:
 	void GatherImportanceVolumes();
 	
 	void BackgroundTick();
+	void AddRelevantStaticLightGUIDs(FQuantizedLightmapData* QuantizedLightmapData, const FBoxSphereBounds& WorldBounds);
 	void ApplyFinishedLightmapsToWorld();
 	void RemoveAllComponents();
 
