GPULM: 64bit hash table with contention detection. This should allow IC to cover a space of 2^19 * cache entry size. Also, no corner rejection during IC building passes

diff --git a/Engine/Plugins/Experimental/GPULightmass/Shaders/Private/IrradianceCacheVisualization.usf b/Engine/Plugins/Experimental/GPULightmass/Shaders/Private/IrradianceCacheVisualization.usf
--- a/Engine/Plugins/Experimental/GPULightmass/Shaders/Private/IrradianceCacheVisualization.usf
+++ b/Engine/Plugins/Experimental/GPULightmass/Shaders/Private/IrradianceCacheVisualization.usf
@@ -25,10 +25,9 @@ void VisualizeIrradianceCachePS(
 			
 	uint NearestRecordIndex = 0;
 	float3 RecordIrradiance;
-			
-	uint HashKey = EncodeICHashKey(WorldPosition, ShadingNormal, IrradianceCachingParameters.Spacing);
+
 	uint Index;
-	if (ICHashTableFind(HashKey, Index))
+	if (ICHashTableFind(EncodeICHashKey(WorldPosition, ShadingNormal, IrradianceCachingParameters.Spacing), Index))
 	{
 		uint RecordIndex = IrradianceCachingParameters.RWHashToIndex[Index];
 		uint4 RecordIrradianceAndSampleCount = IrradianceCachingParameters.IrradianceCacheRecords[RecordIndex];
diff --git a/Engine/Plugins/Experimental/GPULightmass/Shaders/Private/IrradianceCachingCommon.ush b/Engine/Plugins/Experimental/GPULightmass/Shaders/Private/IrradianceCachingCommon.ush
--- a/Engine/Plugins/Experimental/GPULightmass/Shaders/Private/IrradianceCachingCommon.ush
+++ b/Engine/Plugins/Experimental/GPULightmass/Shaders/Private/IrradianceCachingCommon.ush
@@ -6,30 +6,74 @@ struct FFinalGatherHitPoint
 	float3 WorldNormal;
 };
 
-bool ICHashTableAdd( uint Key, out uint Index )
+uint ICKeyMix(uint2 Key)
+{
+	return MurmurMix(MurmurAdd(Key.x, Key.y));
+}
+
+bool ICHashTableAdd(uint2 Key, out uint Index)
 {
 	// Zero is reserved as invalid
 	Key++;
 
 	LOOP
 	[allow_uav_condition]
-	for( Index = MurmurMix( Key );; Index++ )
+	for(Index = ICKeyMix(Key); ; Index++)
 	{
 		Index = Index % IrradianceCachingParameters.HashTableSize;
+		
+		uint2 StoredKey = 0;
+
+		bool bAtomicReadSuccessful = false;
+		uint SemaphoreState;
+		InterlockedCompareExchange(IrradianceCachingParameters.HashTableSemaphore[Index], 0, 1, SemaphoreState);
+		if (SemaphoreState == 0 || SemaphoreState == 1)
+		{
+			StoredKey = IrradianceCachingParameters.RWHashTable[Index];
+			bAtomicReadSuccessful = true;
+			
+			if (SemaphoreState == 0)
+			{
+				InterlockedCompareExchange(IrradianceCachingParameters.HashTableSemaphore[Index], 1, 0, SemaphoreState);
+			}
+		}
 
-		uint StoredKey = IrradianceCachingParameters.RWHashTable[ Index ];
-		if( StoredKey != Key )
+		if (!bAtomicReadSuccessful)
 		{
-			if( StoredKey != 0 )
+			return false;
+		}
+		
+		if(any(StoredKey != Key))
+		{
+			if(any(StoredKey != 0))
 				continue;
 
-			uint PrevKey;
-			InterlockedCompareExchange( IrradianceCachingParameters.RWHashTable[ Index ], 0, Key, PrevKey );
+			uint2 PrevKey = 0xffffffff;
 			
-			if( PrevKey == 0 )
-				return true;
-			else if( PrevKey != Key )
-				continue;
+			bool bAtomicWriteSuccessful = false;
+			InterlockedCompareExchange(IrradianceCachingParameters.HashTableSemaphore[Index], 0, 2, SemaphoreState);
+			if (SemaphoreState == 0)
+			{
+				bAtomicWriteSuccessful = true;
+				PrevKey = IrradianceCachingParameters.RWHashTable[ Index ];
+				if (all(PrevKey == 0))
+				{
+					IrradianceCachingParameters.RWHashTable[Index] = Key;
+				}
+				InterlockedCompareExchange(IrradianceCachingParameters.HashTableSemaphore[Index], 2, 0, SemaphoreState);
+			}
+
+			if (bAtomicWriteSuccessful)
+			{
+				if(all(PrevKey == 0))
+					return true;
+				else if(any(PrevKey != Key))
+					continue;
+			}
+			else
+			{
+				return false;
+			}
 		}
 
 		break;
@@ -40,21 +84,41 @@ bool ICHashTableAdd( uint Key, out uint Index )
 
 // Returns true if key is found.
 // Index output is the hash table bucket this key is stored in if found.
-bool ICHashTableFind( uint Key, out uint Index )
+bool ICHashTableFind(uint2 Key, out uint Index)
 {
 	// Zero is reserved as invalid
 	Key++;
 
 	LOOP
 	[allow_uav_condition]
-	for( Index = MurmurMix( Key );; Index++ )
+	for(Index = ICKeyMix(Key); ; Index++)
 	{
 		Index = Index % IrradianceCachingParameters.HashTableSize;
 
-		uint StoredKey = IrradianceCachingParameters.RWHashTable[ Index ];
-		if( StoredKey != Key )
+		uint2 StoredKey = 0;
+
+		bool bAtomicReadSuccessful = false;
+		uint SemaphoreState;
+		InterlockedCompareExchange(IrradianceCachingParameters.HashTableSemaphore[Index], 0, 1, SemaphoreState);
+		if (SemaphoreState == 0 || SemaphoreState == 1)
+		{
+			StoredKey = IrradianceCachingParameters.RWHashTable[Index];
+			bAtomicReadSuccessful = true;
+			
+			if (SemaphoreState == 0)
+			{
+				InterlockedCompareExchange(IrradianceCachingParameters.HashTableSemaphore[Index], 1, 0, SemaphoreState);
+			}
+		}
+		
+		if (!bAtomicReadSuccessful)
+		{
+			return false;
+		}
+		
+		if(any(StoredKey != Key))
 		{
-			if( StoredKey != 0 )
+			if(any(StoredKey != 0))
 				continue;
 		}
 		else
@@ -70,8 +134,8 @@ bool ICHashTableFind( uint Key, out uint Index )
 
 uint3 EncodeVoxelKey(float3 VoxelPos)
 {
-	int3 Signed = int3( VoxelPos ) + 0x1000000;
-	uint3 Voxel = uint3( Signed ) & 0xfffffff;
+	int3 Signed = int3( VoxelPos ) + 0x3ffff;
+	uint3 Voxel = uint3( Signed ) & 0x7ffff;
 
 	uint3 Key;
 	Key.x = Voxel.x;
@@ -119,16 +183,29 @@ uint EncodeNormalBits(float3 WorldNormal)
 			;
 }
 
-uint EncodeICHashKey(float3 WorldPosition, float3 WorldNormal, float Spacing)
+uint2 EncodeICHashKey(float3 WorldPosition, float3 WorldNormal, float Spacing)
 {
 	float3 PosKey = EncodePositionKey(WorldPosition, WorldNormal, Spacing); // RT_LWC_TODO
 	uint3 Key = EncodeVoxelKey(PosKey);
-	uint NormalDirectionBits = EncodeNormalBits(WorldNormal);			
-	uint HashKey = 0;		
-	HashKey = MurmurAdd( HashKey, NormalDirectionBits << 26);
-	HashKey = MurmurAdd( HashKey, Key.x << 17);
-	HashKey = MurmurAdd( HashKey, Key.y << 8);
-	HashKey = MurmurAdd( HashKey, Key.z << 0);
+	uint NormalDirectionBits = EncodeNormalBits(WorldNormal);
+	
+	uint2 HashKey = uint2(0, 0);
+
+	// We have 64 bits in total
+	// x: 31 - 26: 06 bits for normal direction
+	// x: 25 - 07: 19 bits for world position x
+	// x: 06 - 00: 07 bits for world position y (19 bits in total)
+	// y: 31 - 20: 12 bits for world position y (19 bits in total)
+	// y: 19 - 01: 19 bits for world position z
+	// y: 00: unused
+	
+	HashKey.x = HashKey.x | NormalDirectionBits << 26;
+	HashKey.x = HashKey.x | ((Key.x & 0x7fff) << 7);
+	HashKey.x = HashKey.x | ((Key.y & 0x7f) << 0);
+	
+	HashKey.y = HashKey.y | (Key.y & 0xfff) << 20;
+	HashKey.y = HashKey.y | (Key.z & 0x7fff) << 1;
+	
 	return HashKey;
 }
 
@@ -150,9 +227,8 @@ uint EncodeICHashKey(float3 WorldPosition, float3 WorldNormal, float Spacing)
 
 void EmitGeometryHitPoint(FFinalGatherHitPoint HitPoint, uint Size)
 {
-	uint HashKey = EncodeICHashKey(HitPoint.WorldPosition, HitPoint.WorldNormal, Size);
 	uint Index;
-	if (ICHashTableAdd(HashKey, Index))
+	if (ICHashTableAdd(EncodeICHashKey(HitPoint.WorldPosition, HitPoint.WorldNormal, Size), Index))
 	{
 		uint RecordIndex = 0;
 		InterlockedAdd(IrradianceCachingParameters.RecordAllocator[0], 1, RecordIndex);
diff --git a/Engine/Plugins/Experimental/GPULightmass/Shaders/Private/LightmapPathTracing.usf b/Engine/Plugins/Experimental/GPULightmass/Shaders/Private/LightmapPathTracing.usf
--- a/Engine/Plugins/Experimental/GPULightmass/Shaders/Private/LightmapPathTracing.usf
+++ b/Engine/Plugins/Experimental/GPULightmass/Shaders/Private/LightmapPathTracing.usf
@@ -410,16 +410,7 @@ void PathTracingKernel(
 		float3 TexelPosition = TranslatedWorldPosition - LWCHackToFloat(PrimaryView.PreViewTranslation);
 		float3 HitPosition = Payload.TranslatedWorldPos - LWCHackToFloat(PrimaryView.PreViewTranslation);
 		float DistanceFromTexelToHit = length(HitPosition - TexelPosition);
-		bool bRejectedByCornerRejection = DistanceFromTexelToHit < IrradianceCachingParameters.Spacing * IrradianceCachingParameters.CornerRejection;
-			
-		// If we're in IC building passes and the entry is going to be rejected, terminate immediately as the path isn't going to generate any useful result
-		if (RenderPassIndex < GetIrradianceCachingQuality())
-		{
-			if (Bounce == 1 && bRejectedByCornerRejection)
-			{
-				break;
-			}
-		}
+		bool bRejectedByCornerRejection = RenderPassIndex >= GetIrradianceCachingQuality() && DistanceFromTexelToHit < IrradianceCachingParameters.Spacing * IrradianceCachingParameters.CornerRejection;
 		
 		if (NearestCacheEntryID == -1 && Bounce >= 1 && !bRejectedByCornerRejection)
 		{
@@ -428,10 +419,9 @@ void PathTracingKernel(
 			
 			uint NearestRecordIndex = 0;
 			float3 RecordIrradiance;
-				
-			uint HashKey = EncodeICHashKey(HitPosition, Payload.WorldNormal, IrradianceCachingParameters.Spacing);
+
 			uint Index;
-			if (ICHashTableFind(HashKey, Index))
+			if (ICHashTableFind(EncodeICHashKey(HitPosition, Payload.WorldNormal, IrradianceCachingParameters.Spacing), Index))
 			{
 				uint RecordIndex = IrradianceCachingParameters.RWHashToIndex[Index];
 
diff --git a/Engine/Plugins/Experimental/GPULightmass/Source/GPULightmass/Private/IrradianceCaching.cpp b/Engine/Plugins/Experimental/GPULightmass/Source/GPULightmass/Private/IrradianceCaching.cpp
--- a/Engine/Plugins/Experimental/GPULightmass/Source/GPULightmass/Private/IrradianceCaching.cpp
+++ b/Engine/Plugins/Experimental/GPULightmass/Source/GPULightmass/Private/IrradianceCaching.cpp
@@ -31,17 +31,20 @@ FIrradianceCache::FIrradianceCache(int32 Quality, float Spacing, float CornerRej
 
 	{
 		{
-			HashTable.Initialize(TEXT("ICHashTable"), sizeof(uint32), HashTableSize, PF_R32_UINT, BUF_UnorderedAccess | BUF_ShaderResource);
+			HashTable.Initialize(TEXT("ICHashTable"), sizeof(uint32) * 2, HashTableSize, PF_R32G32_UINT, BUF_UnorderedAccess | BUF_ShaderResource);
 			HashToIndex.Initialize(TEXT("ICHashToIndex"), sizeof(uint32), HashTableSize, PF_R32_UINT, BUF_UnorderedAccess | BUF_ShaderResource);
 			IndexToHash.Initialize(TEXT("ICIndexToHash"), sizeof(uint32), HashTableSize, PF_R32_UINT, BUF_UnorderedAccess | BUF_ShaderResource);
+			HashTableSemaphore.Initialize(TEXT("ICHashTableSemaphore"), sizeof(uint32), HashTableSize, PF_R32_UINT, BUF_UnorderedAccess | BUF_ShaderResource);
 
 			IrradianceCacheTotalBytes += HashTable.NumBytes;
 			IrradianceCacheTotalBytes += HashToIndex.NumBytes;
 			IrradianceCacheTotalBytes += IndexToHash.NumBytes;
+			IrradianceCacheTotalBytes += HashTableSemaphore.NumBytes;
 
 			RHICmdList.ClearUAVUint(HashTable.UAV, FUintVector4(0, 0, 0, 0));
 			RHICmdList.ClearUAVUint(HashToIndex.UAV, FUintVector4(0, 0, 0, 0));
 			RHICmdList.ClearUAVUint(IndexToHash.UAV, FUintVector4(0, 0, 0, 0));
+			RHICmdList.ClearUAVUint(HashTableSemaphore.UAV, FUintVector4(0, 0, 0, 0));
 		}
 		{
 			RecordAllocator.Initialize(TEXT("ICAllocator"), sizeof(uint32), 1, PF_R32_UINT, BUF_UnorderedAccess | BUF_ShaderResource);
@@ -64,5 +67,6 @@ FIrradianceCache::FIrradianceCache(int32 Quality, float Spacing, float CornerRej
 	IrradianceCachingParameters.RWHashToIndex = HashToIndex.UAV;
 	IrradianceCachingParameters.RWIndexToHash = IndexToHash.UAV;
 	IrradianceCachingParameters.RecordAllocator = RecordAllocator.UAV;
+	IrradianceCachingParameters.HashTableSemaphore = HashTableSemaphore.UAV;
 	IrradianceCachingParametersUniformBuffer = TUniformBufferRef<FIrradianceCachingParameters>::CreateUniformBufferImmediate(IrradianceCachingParameters, EUniformBufferUsage::UniformBuffer_MultiFrame);
 }
diff --git a/Engine/Plugins/Experimental/GPULightmass/Source/GPULightmass/Private/IrradianceCaching.h b/Engine/Plugins/Experimental/GPULightmass/Source/GPULightmass/Private/IrradianceCaching.h
--- a/Engine/Plugins/Experimental/GPULightmass/Source/GPULightmass/Private/IrradianceCaching.h
+++ b/Engine/Plugins/Experimental/GPULightmass/Source/GPULightmass/Private/IrradianceCaching.h
@@ -13,10 +13,11 @@ BEGIN_UNIFORM_BUFFER_STRUCT(FIrradianceCachingParameters, )
 	SHADER_PARAMETER(float, Spacing)
 	SHADER_PARAMETER(float, CornerRejection)
 	SHADER_PARAMETER_UAV(RWStructuredBuffer<uint4>, IrradianceCacheRecords)
-	SHADER_PARAMETER_UAV(RWStructuredBuffer<uint>, RWHashTable)
+	SHADER_PARAMETER_UAV(RWStructuredBuffer<uint2>, RWHashTable)
 	SHADER_PARAMETER_UAV(RWStructuredBuffer<uint>, RWHashToIndex)
 	SHADER_PARAMETER_UAV(RWStructuredBuffer<uint>, RWIndexToHash)
 	SHADER_PARAMETER_UAV(RWStructuredBuffer<uint>, RecordAllocator)
+	SHADER_PARAMETER_UAV(RWStructuredBuffer<uint>, HashTableSemaphore)
 END_UNIFORM_BUFFER_STRUCT()
 
 struct FIrradianceCache
@@ -34,6 +35,7 @@ struct FIrradianceCache
 	FRWBuffer HashToIndex;
 	FRWBuffer IndexToHash;
 	FRWBuffer RecordAllocator;
+	FRWBuffer HashTableSemaphore;
 
 	int32 CurrentRevision = 0;
 };
