GPULM: Fix bug where VT lightmap preview tiles are not getting populated properly

diff --git a/Engine/Plugins/Experimental/GPULightmass/Source/GPULightmass/Private/LightmapRenderer.cpp b/Engine/Plugins/Experimental/GPULightmass/Source/GPULightmass/Private/LightmapRenderer.cpp
--- a/Engine/Plugins/Experimental/GPULightmass/Source/GPULightmass/Private/LightmapRenderer.cpp
+++ b/Engine/Plugins/Experimental/GPULightmass/Source/GPULightmass/Private/LightmapRenderer.cpp
@@ -1676,7 +1676,7 @@ void FLightmapRenderer::Finalize(FRDGBuilder& GraphBuilder)
 					}
 
 					{
-						FRDGTextureRef RenderTargetTileAtlas = GraphBuilder.RegisterExternalTexture(OutputRenderTargets[0], *(FString(TEXT("GPULightmassRenderTargetTileAtlas_")) + StagingTextureNames[RenderTargetIndex]));
+						FRDGTextureRef RenderTargetTileAtlas = GraphBuilder.RegisterExternalTexture(OutputRenderTargets[RenderTargetIndex], *(FString(TEXT("GPULightmassRenderTargetTileAtlas_")) + StagingTextureNames[RenderTargetIndex]));
 
 						FUploadConvergedLightmapTilesCS::FParameters* PassParameters = GraphBuilder.AllocParameters<FUploadConvergedLightmapTilesCS::FParameters>();
 
