GPULM: Implement skylight bent normal calculation in VLM. This completes the support for stationary skylights.

diff --git a/Engine/Plugins/Experimental/GPULightmass/Shaders/Private/BrickAllocationManagement.usf b/Engine/Plugins/Experimental/GPULightmass/Shaders/Private/BrickAllocationManagement.usf
--- a/Engine/Plugins/Experimental/GPULightmass/Shaders/Private/BrickAllocationManagement.usf
+++ b/Engine/Plugins/Experimental/GPULightmass/Shaders/Private/BrickAllocationManagement.usf
@@ -133,6 +133,7 @@ Texture3D<float4> SHCoefficients0G;
 Texture3D<float4> SHCoefficients1G;
 Texture3D<float4> SHCoefficients0B;
 Texture3D<float4> SHCoefficients1B;
+Texture3D<float4> SkyBentNormal;
 Texture3D<UNORM float> DirectionalLightShadowing;
 
 RWTexture3D<float3> OutAmbientVector;
@@ -142,6 +143,7 @@ RWTexture3D<UNORM float4> OutSHCoefficients0G;
 RWTexture3D<UNORM float4> OutSHCoefficients1G;
 RWTexture3D<UNORM float4> OutSHCoefficients0B;
 RWTexture3D<UNORM float4> OutSHCoefficients1B;
+RWTexture3D<UNORM float4> OutSkyBentNormal;
 RWTexture3D<UNORM float> OutDirectionalLightShadowing;
 
 uint FrameNumber;
@@ -369,6 +371,7 @@ void StitchBorderCS(uint3 BrickID : SV_GroupID, uint3 CellPosInBrick : SV_GroupT
 		OutSHCoefficients1G[VoxelPos] = DilateWithValidityMask(OutSHCoefficients1G, AmbientVector, IndirectionTextureUVW);
 		OutSHCoefficients0B[VoxelPos] = DilateWithValidityMask(OutSHCoefficients0B, AmbientVector, IndirectionTextureUVW);
 		OutSHCoefficients1B[VoxelPos] = DilateWithValidityMask(OutSHCoefficients1B, AmbientVector, IndirectionTextureUVW);
+		OutSkyBentNormal[VoxelPos] = DilateWithValidityMask(OutSkyBentNormal, AmbientVector, IndirectionTextureUVW);
 		OutDirectionalLightShadowing[VoxelPos] = DilateWithValidityMask(OutDirectionalLightShadowing, AmbientVector, IndirectionTextureUVW);
 	}
 	
@@ -392,6 +395,7 @@ void StitchBorderCS(uint3 BrickID : SV_GroupID, uint3 CellPosInBrick : SV_GroupT
 					OutSHCoefficients1G[VoxelPos] = ManualTrilinearFilter(OutSHCoefficients1G, BrickUVW);
 					OutSHCoefficients0B[VoxelPos] = ManualTrilinearFilter(OutSHCoefficients0B, BrickUVW);
 					OutSHCoefficients1B[VoxelPos] = ManualTrilinearFilter(OutSHCoefficients1B, BrickUVW);
+					   OutSkyBentNormal[VoxelPos] = ManualTrilinearFilter(OutSkyBentNormal, BrickUVW);
 					OutDirectionalLightShadowing[VoxelPos] = ManualTrilinearFilter(OutDirectionalLightShadowing, BrickUVW);
 				}
 			}
@@ -434,6 +438,7 @@ void FinalizeBrickResultsCS(uint3 BrickID : SV_GroupID, uint3 CellPosInBrick : S
 		OutSHCoefficients1G[VoxelPos] = (SHCoefficients1G[VoxelPos] / SampleCount) * InvAmbient.g * CoefficientNormalizationScale1 * 0.5f + 0.5f;
 		OutSHCoefficients0B[VoxelPos] = (SHCoefficients0B[VoxelPos] / SampleCount) * InvAmbient.b * CoefficientNormalizationScale0 * 0.5f + 0.5f;
 		OutSHCoefficients1B[VoxelPos] = (SHCoefficients1B[VoxelPos] / SampleCount) * InvAmbient.b * CoefficientNormalizationScale1 * 0.5f + 0.5f;
+		OutSkyBentNormal[VoxelPos] = SkyBentNormal[VoxelPos] / SampleCount * 0.5f + 0.5f;
 		OutDirectionalLightShadowing[VoxelPos] = DirectionalLightShadowing[VoxelPos];
 	}
 }
diff --git a/Engine/Plugins/Experimental/GPULightmass/Shaders/Private/LightmapPathTracing.usf b/Engine/Plugins/Experimental/GPULightmass/Shaders/Private/LightmapPathTracing.usf
--- a/Engine/Plugins/Experimental/GPULightmass/Shaders/Private/LightmapPathTracing.usf
+++ b/Engine/Plugins/Experimental/GPULightmass/Shaders/Private/LightmapPathTracing.usf
@@ -1039,6 +1039,7 @@ RWTexture3D<float4> SHCoefficients0G;
 RWTexture3D<float4> SHCoefficients1G;
 RWTexture3D<float4> SHCoefficients0B;
 RWTexture3D<float4> SHCoefficients1B;
+RWTexture3D<float4> SkyBentNormal;
 RWTexture3D<UNORM float> DirectionalLightShadowing;
 
 uint MortonEncode3(uint3 Pixel)
@@ -1279,6 +1280,8 @@ void VolumetricLightmapPathTracingMainRG()
 		
 		SHCoefficients0B[VoxelPos] = float4(0, 0, 0, 0);
 		SHCoefficients1B[VoxelPos] = float4(0, 0, 0, 0);
+
+		SkyBentNormal[VoxelPos] = float4(0, 0, 0, 0);
 	}
 	
 	if (bIsValidSample)
@@ -1309,6 +1312,13 @@ void VolumetricLightmapPathTracingMainRG()
 			
 			SHCoefficients0B[VoxelPos] += float4(SH.B.V0.yzw, SH.B.V1.x);
 			SHCoefficients1B[VoxelPos] += float4(SH.B.V1.yzw, SH.B.V2);
+			
+			if (EffectiveRenderPassIndex % 2 == 0)
+			{
+				// Every even sample is upper hemisphere, odd sample is lower hemisphere
+				// SampleCount stored is twice as the actual full sphere sample count
+				SkyBentNormal[VoxelPos].rgb += SkyLightBentNormal * 2;
+			}
 		}
 
 		#ifndef LIGHTMAP_PATH_TRACING_MAIN_RG
diff --git a/Engine/Plugins/Experimental/GPULightmass/Source/GPULightmass/Private/LightmapRayTracing.h b/Engine/Plugins/Experimental/GPULightmass/Source/GPULightmass/Private/LightmapRayTracing.h
--- a/Engine/Plugins/Experimental/GPULightmass/Source/GPULightmass/Private/LightmapRayTracing.h
+++ b/Engine/Plugins/Experimental/GPULightmass/Source/GPULightmass/Private/LightmapRayTracing.h
@@ -129,6 +129,7 @@ class FVolumetricLightmapPathTracingRGS : public FGlobalShader
 		SHADER_PARAMETER_UAV(RWTexture3D<float4>, SHCoefficients1G)
 		SHADER_PARAMETER_UAV(RWTexture3D<float4>, SHCoefficients0B)
 		SHADER_PARAMETER_UAV(RWTexture3D<float4>, SHCoefficients1B)
+		SHADER_PARAMETER_UAV(RWTexture3D<float4>, SkyBentNormal)
 		SHADER_PARAMETER_UAV(RWTexture3D<float>, DirectionalLightShadowing)
 
 		SHADER_PARAMETER_RDG_BUFFER_SRV(StructuredBuffer<FPathTracingLight>, SceneLights)
diff --git a/Engine/Plugins/Experimental/GPULightmass/Source/GPULightmass/Private/VolumetricLightmap.cpp b/Engine/Plugins/Experimental/GPULightmass/Source/GPULightmass/Private/VolumetricLightmap.cpp
--- a/Engine/Plugins/Experimental/GPULightmass/Source/GPULightmass/Private/VolumetricLightmap.cpp
+++ b/Engine/Plugins/Experimental/GPULightmass/Source/GPULightmass/Private/VolumetricLightmap.cpp
@@ -46,11 +46,8 @@ void InitializeBrickData(FIntVector BrickDataDimensions, FVolumetricLightmapBric
 		BrickData.SHCoefficients[i].CreateUAV();
 	}
 
-	//if (BrickData.SkyBentNormal.Texture.IsValid())
-	//{
-	//	SkyBentNormal.CreateTargetTexture(BrickDataDimensions);
-	//	SkyBentNormal.CreateUAV();
-	//}
+	BrickData.SkyBentNormal.CreateTargetTexture(BrickDataDimensions);
+	BrickData.SkyBentNormal.CreateUAV();
 
 	BrickData.DirectionalLightShadowing.CreateTargetTexture(BrickDataDimensions);
 	BrickData.DirectionalLightShadowing.CreateUAV();
@@ -523,7 +520,9 @@ void FVolumetricLightmapRenderer::BackgroundTick()
 				FRHITransitionInfo(VolumetricLightmapData.BrickData.SHCoefficients[2].UAV, ERHIAccess::Unknown, ERHIAccess::UAVCompute),
 				FRHITransitionInfo(VolumetricLightmapData.BrickData.SHCoefficients[3].UAV, ERHIAccess::Unknown, ERHIAccess::UAVCompute),
 				FRHITransitionInfo(VolumetricLightmapData.BrickData.SHCoefficients[4].UAV, ERHIAccess::Unknown, ERHIAccess::UAVCompute),
-				FRHITransitionInfo(VolumetricLightmapData.BrickData.SHCoefficients[5].UAV, ERHIAccess::Unknown, ERHIAccess::UAVCompute)
+				FRHITransitionInfo(VolumetricLightmapData.BrickData.SHCoefficients[5].UAV, ERHIAccess::Unknown, ERHIAccess::UAVCompute),
+				FRHITransitionInfo(VolumetricLightmapData.BrickData.SkyBentNormal.UAV, ERHIAccess::Unknown, ERHIAccess::UAVCompute),
+				FRHITransitionInfo(VolumetricLightmapData.BrickData.DirectionalLightShadowing.UAV, ERHIAccess::Unknown, ERHIAccess::UAVCompute),
 				}
 			);
 		}
@@ -563,6 +562,7 @@ void FVolumetricLightmapRenderer::BackgroundTick()
 					PassParameters->SHCoefficients1G = AccumulationBrickData.SHCoefficients[3].UAV;
 					PassParameters->SHCoefficients0B = AccumulationBrickData.SHCoefficients[4].UAV;
 					PassParameters->SHCoefficients1B = AccumulationBrickData.SHCoefficients[5].UAV;
+					PassParameters->SkyBentNormal = AccumulationBrickData.SkyBentNormal.UAV;
 					PassParameters->DirectionalLightShadowing = AccumulationBrickData.DirectionalLightShadowing.UAV;
 					PassParameters->ViewUniformBuffer = Scene->ReferenceView->ViewUniformBuffer;
 					PassParameters->IrradianceCachingParameters = Scene->IrradianceCache->IrradianceCachingParametersUniformBuffer;
@@ -639,6 +639,8 @@ void FVolumetricLightmapRenderer::BackgroundTick()
 						FRHITransitionInfo(AccumulationBrickData.SHCoefficients[3].UAV, ERHIAccess::UAVCompute, ERHIAccess::UAVCompute),
 						FRHITransitionInfo(AccumulationBrickData.SHCoefficients[4].UAV, ERHIAccess::UAVCompute, ERHIAccess::UAVCompute),
 						FRHITransitionInfo(AccumulationBrickData.SHCoefficients[5].UAV, ERHIAccess::UAVCompute, ERHIAccess::UAVCompute),
+						FRHITransitionInfo(AccumulationBrickData.SkyBentNormal.UAV, ERHIAccess::UAVCompute, ERHIAccess::UAVCompute),
+						FRHITransitionInfo(AccumulationBrickData.DirectionalLightShadowing.UAV, ERHIAccess::UAVCompute, ERHIAccess::UAVCompute),
 						}
 					);
 				}
@@ -657,6 +659,7 @@ void FVolumetricLightmapRenderer::BackgroundTick()
 				PassParameters->SHCoefficients1G = AccumulationBrickData.SHCoefficients[3].Texture;
 				PassParameters->SHCoefficients0B = AccumulationBrickData.SHCoefficients[4].Texture;
 				PassParameters->SHCoefficients1B = AccumulationBrickData.SHCoefficients[5].Texture;
+				PassParameters->SkyBentNormal = AccumulationBrickData.SkyBentNormal.Texture;
 				PassParameters->DirectionalLightShadowing = AccumulationBrickData.DirectionalLightShadowing.Texture;
 				PassParameters->OutAmbientVector = VolumetricLightmapData.BrickData.AmbientVector.UAV;
 				PassParameters->OutSHCoefficients0R = VolumetricLightmapData.BrickData.SHCoefficients[0].UAV;
@@ -665,6 +668,7 @@ void FVolumetricLightmapRenderer::BackgroundTick()
 				PassParameters->OutSHCoefficients1G = VolumetricLightmapData.BrickData.SHCoefficients[3].UAV;
 				PassParameters->OutSHCoefficients0B = VolumetricLightmapData.BrickData.SHCoefficients[4].UAV;
 				PassParameters->OutSHCoefficients1B = VolumetricLightmapData.BrickData.SHCoefficients[5].UAV;
+				PassParameters->OutSkyBentNormal = VolumetricLightmapData.BrickData.SkyBentNormal.UAV;
 				PassParameters->OutDirectionalLightShadowing = VolumetricLightmapData.BrickData.DirectionalLightShadowing.UAV;
 
 				FComputeShaderUtils::AddPass(GraphBuilder, RDG_EVENT_NAME("FinalizeBrickResults"), ComputeShader, PassParameters, FIntVector(BricksToCalcThisFrame, 1, 1));
@@ -680,6 +684,8 @@ void FVolumetricLightmapRenderer::BackgroundTick()
 						FRHITransitionInfo(AccumulationBrickData.SHCoefficients[3].UAV, ERHIAccess::UAVCompute, ERHIAccess::UAVCompute),
 						FRHITransitionInfo(AccumulationBrickData.SHCoefficients[4].UAV, ERHIAccess::UAVCompute, ERHIAccess::UAVCompute),
 						FRHITransitionInfo(AccumulationBrickData.SHCoefficients[5].UAV, ERHIAccess::UAVCompute, ERHIAccess::UAVCompute),
+						FRHITransitionInfo(AccumulationBrickData.SkyBentNormal.UAV, ERHIAccess::UAVCompute, ERHIAccess::UAVCompute),
+						FRHITransitionInfo(AccumulationBrickData.DirectionalLightShadowing.UAV, ERHIAccess::UAVCompute, ERHIAccess::UAVCompute),
 						}
 					);
 				}
@@ -708,6 +714,7 @@ void FVolumetricLightmapRenderer::BackgroundTick()
 					PassParameters->OutSHCoefficients1G = VolumetricLightmapData.BrickData.SHCoefficients[3].UAV;
 					PassParameters->OutSHCoefficients0B = VolumetricLightmapData.BrickData.SHCoefficients[4].UAV;
 					PassParameters->OutSHCoefficients1B = VolumetricLightmapData.BrickData.SHCoefficients[5].UAV;
+					PassParameters->OutSkyBentNormal = VolumetricLightmapData.BrickData.SkyBentNormal.UAV;
 					PassParameters->OutDirectionalLightShadowing = VolumetricLightmapData.BrickData.DirectionalLightShadowing.UAV;
 
 					FComputeShaderUtils::AddPass(GraphBuilder, RDG_EVENT_NAME("VolumetricLightmapStitching %d bricks", BricksToCalcThisFrame), ComputeShader, PassParameters, FIntVector(BricksToCalcThisFrame, 1, 1));
@@ -722,7 +729,9 @@ void FVolumetricLightmapRenderer::BackgroundTick()
 							FRHITransitionInfo(AccumulationBrickData.SHCoefficients[2].UAV, ERHIAccess::UAVCompute, ERHIAccess::UAVCompute),
 							FRHITransitionInfo(AccumulationBrickData.SHCoefficients[3].UAV, ERHIAccess::UAVCompute, ERHIAccess::UAVCompute),
 							FRHITransitionInfo(AccumulationBrickData.SHCoefficients[4].UAV, ERHIAccess::UAVCompute, ERHIAccess::UAVCompute),
-							FRHITransitionInfo(AccumulationBrickData.SHCoefficients[5].UAV, ERHIAccess::UAVCompute, ERHIAccess::UAVCompute)
+							FRHITransitionInfo(AccumulationBrickData.SHCoefficients[5].UAV, ERHIAccess::UAVCompute, ERHIAccess::UAVCompute),
+							FRHITransitionInfo(AccumulationBrickData.SkyBentNormal.UAV, ERHIAccess::UAVCompute, ERHIAccess::UAVCompute),
+							FRHITransitionInfo(AccumulationBrickData.DirectionalLightShadowing.UAV, ERHIAccess::UAVCompute, ERHIAccess::UAVCompute),
 							}
 						);
 					}
diff --git a/Engine/Plugins/Experimental/GPULightmass/Source/GPULightmass/Private/VolumetricLightmapVoxelization.h b/Engine/Plugins/Experimental/GPULightmass/Source/GPULightmass/Private/VolumetricLightmapVoxelization.h
--- a/Engine/Plugins/Experimental/GPULightmass/Source/GPULightmass/Private/VolumetricLightmapVoxelization.h
+++ b/Engine/Plugins/Experimental/GPULightmass/Source/GPULightmass/Private/VolumetricLightmapVoxelization.h
@@ -359,6 +359,7 @@ class FStitchBorderCS : public FGlobalShader
 		SHADER_PARAMETER_UAV(RWTexture3D<float4>, OutSHCoefficients1G)
 		SHADER_PARAMETER_UAV(RWTexture3D<float4>, OutSHCoefficients0B)
 		SHADER_PARAMETER_UAV(RWTexture3D<float4>, OutSHCoefficients1B)
+		SHADER_PARAMETER_UAV(RWTexture3D<float4>, OutSkyBentNormal)
 		SHADER_PARAMETER_UAV(RWTexture3D<float>, OutDirectionalLightShadowing)
 	END_SHADER_PARAMETER_STRUCT()
 };
@@ -389,6 +390,7 @@ class FFinalizeBrickResultsCS : public FGlobalShader
 		SHADER_PARAMETER_TEXTURE(Texture3D<float4>, SHCoefficients1G)
 		SHADER_PARAMETER_TEXTURE(Texture3D<float4>, SHCoefficients0B)
 		SHADER_PARAMETER_TEXTURE(Texture3D<float4>, SHCoefficients1B)
+		SHADER_PARAMETER_TEXTURE(Texture3D<float4>, SkyBentNormal)
 		SHADER_PARAMETER_TEXTURE(Texture3D<float>, DirectionalLightShadowing)
 		SHADER_PARAMETER_UAV(RWTexture3D<float3>, OutAmbientVector)
 		SHADER_PARAMETER_UAV(RWTexture3D<float4>, OutSHCoefficients0R)
@@ -397,6 +399,7 @@ class FFinalizeBrickResultsCS : public FGlobalShader
 		SHADER_PARAMETER_UAV(RWTexture3D<float4>, OutSHCoefficients1G)
 		SHADER_PARAMETER_UAV(RWTexture3D<float4>, OutSHCoefficients0B)
 		SHADER_PARAMETER_UAV(RWTexture3D<float4>, OutSHCoefficients1B)
+		SHADER_PARAMETER_UAV(RWTexture3D<float4>, OutSkyBentNormal)
 		SHADER_PARAMETER_UAV(RWTexture3D<float>, OutDirectionalLightShadowing)
 	END_SHADER_PARAMETER_STRUCT()
 };
