Converted volumetric lightmap textures to RDG.

diff --git a/Engine/Plugins/Experimental/GPULightmass/Source/GPULightmass/Private/VolumetricLightmap.cpp b/Engine/Plugins/Experimental/GPULightmass/Source/GPULightmass/Private/VolumetricLightmap.cpp
--- a/Engine/Plugins/Experimental/GPULightmass/Source/GPULightmass/Private/VolumetricLightmap.cpp
+++ b/Engine/Plugins/Experimental/GPULightmass/Source/GPULightmass/Private/VolumetricLightmap.cpp
@@ -116,7 +116,7 @@ void FVolumetricLightmapRenderer::VoxelizeScene()
 	IndirectionTexture.SafeRelease();
 
 	ReleaseBrickData(VolumetricLightmapData.BrickData);
-	ReleaseBrickData( AccumulationBrickData);
+	ReleaseBrickData(AccumulationBrickData);
 
 	FRHICommandListExecutor::GetImmediateCommandList().ImmediateFlush(EImmediateFlushType::FlushRHIThreadFlushResources);
 
@@ -148,267 +148,270 @@ void FVolumetricLightmapRenderer::VoxelizeScene()
 
 	IndirectionTextureDimensions = TopLevelGridSize * IndirectionCellsPerTopLevelCell;
 
-	{
-		const FPooledRenderTargetDesc Desc = FPooledRenderTargetDesc::CreateVolumeDesc(
-			IndirectionTextureDimensions.X, IndirectionTextureDimensions.Y, IndirectionTextureDimensions.Z,
-			PF_R8G8B8A8_UINT,
-			FClearValueBinding::Black,
-			TexCreate_None,
-			TexCreate_ShaderResource | TexCreate_RenderTargetable | TexCreate_UAV,
-			false);
-
-		GRenderTargetPool.FindFreeElement(FRHICommandListExecutor::GetImmediateCommandList(), Desc, IndirectionTexture, *FString::Printf(TEXT("GPULightmassVLMIndirectionTexture")));
-	}
-
 	VoxelizationVolumeMips.Empty();
 
-	for (int32 Level = 0; Level < MaxRefinementLevels; Level++)
-	{
-		const FPooledRenderTargetDesc Desc = FPooledRenderTargetDesc::CreateVolumeDesc(
-			IndirectionTextureDimensions.X >> (Level * BrickSizeLog2), IndirectionTextureDimensions.Y >> (Level * BrickSizeLog2), IndirectionTextureDimensions.Z >> (Level * BrickSizeLog2),
-			PF_R32_UINT,
-			FClearValueBinding::Black,
-			TexCreate_None,
-			TexCreate_ShaderResource | TexCreate_RenderTargetable | TexCreate_UAV,
-			false);
+	FRHICommandListImmediate& RHICmdList = FRHICommandListExecutor::GetImmediateCommandList();
 
-		VoxelizationVolumeMips.AddDefaulted(1);
+	{
+		FRDGBuilder GraphBuilder(RHICmdList);
+		ON_SCOPE_EXIT{ GraphBuilder.Execute(); };
 
-		// UE_LOG(LogGPULightmass, Log, TEXT("Creating voxelization mip%d %d %d %d"), Level, IndirectionTextureDimensions.X >> (Level * BrickSizeLog2), IndirectionTextureDimensions.Y >> (Level * BrickSizeLog2), IndirectionTextureDimensions.Z >> (Level * BrickSizeLog2));
+		FRDGTextureUAV* IndirectionTextureUAV = nullptr;
 
-		GRenderTargetPool.FindFreeElement(FRHICommandListExecutor::GetImmediateCommandList(), Desc, VoxelizationVolumeMips.Last(0), *FString::Printf(TEXT("GPULightmassVLMVoxelizationVolumeMips%d"), Level));
-	}
+		{
+			const FRDGTextureDesc Desc = FRDGTextureDesc::Create3D(
+				FIntVector(IndirectionTextureDimensions.X, IndirectionTextureDimensions.Y, IndirectionTextureDimensions.Z),
+				PF_R8G8B8A8_UINT,
+				FClearValueBinding::Black,
+				ETextureCreateFlags::ShaderResource | ETextureCreateFlags::RenderTargetable | ETextureCreateFlags::UAV);
+
+			FRDGTexture* Texture = GraphBuilder.CreateTexture(Desc, TEXT("GPULightmassVLMIndirectionTexture"));
+			IndirectionTexture = GraphBuilder.ConvertToExternalTexture(Texture);
+			IndirectionTextureUAV = GraphBuilder.CreateUAV(Texture);
+		}
 
-	VolumetricLightmapData.Bounds = FinalVolume;
-	VolumetricLightmapData.IndirectionTexture.Texture = IndirectionTexture->GetRHI();
-	VolumetricLightmapData.IndirectionTexture.Format = PF_R8G8B8A8_UINT;
-	VolumetricLightmapData.IndirectionTextureDimensions = FIntVector(IndirectionTextureDimensions);
-	VolumetricLightmapData.BrickSize = 4;
+		TArray<FRDGTextureUAV*, FRDGArrayAllocator> VoxelizationVolumeMipUAVs;
 
-	FBox CubeVolume(VolumeMin, VolumeMin + FVector(FMath::Max3(VolumeSize.X, VolumeSize.Y, VolumeSize.Z)));
-	int32 CubeMaxDim = FMath::Max3(IndirectionTextureDimensions.X, IndirectionTextureDimensions.Y, IndirectionTextureDimensions.Z);
+		for (int32 Level = 0; Level < MaxRefinementLevels; Level++)
+		{
+			const FRDGTextureDesc Desc = FRDGTextureDesc::Create3D(
+				FIntVector(IndirectionTextureDimensions.X >> (Level * BrickSizeLog2), IndirectionTextureDimensions.Y >> (Level * BrickSizeLog2), IndirectionTextureDimensions.Z >> (Level * BrickSizeLog2)),
+				PF_R32_UINT,
+				FClearValueBinding::Black,
+				ETextureCreateFlags::ShaderResource | ETextureCreateFlags::RenderTargetable | ETextureCreateFlags::UAV);
 
-	FRHICommandListImmediate& RHICmdList = FRHICommandListExecutor::GetImmediateCommandList();
+			FRDGTexture* Texture = GraphBuilder.CreateTexture(Desc, TEXT("GPULightmassVLMVoxelizationVolumeMips"));
 
-	FMemMark MemMark(FMemStack::Get());
+			VoxelizationVolumeMips.Emplace(GraphBuilder.ConvertToExternalTexture(Texture));
+			VoxelizationVolumeMipUAVs.Emplace(GraphBuilder.CreateUAV(Texture));
+		}
 
-	FRDGBuilder GraphBuilder(RHICmdList);
+		VolumetricLightmapData.Bounds = FinalVolume;
+		VolumetricLightmapData.IndirectionTexture.Texture = IndirectionTexture->GetRHI();
+		VolumetricLightmapData.IndirectionTexture.Format = PF_R8G8B8A8_UINT;
+		VolumetricLightmapData.IndirectionTextureDimensions = FIntVector(IndirectionTextureDimensions);
+		VolumetricLightmapData.BrickSize = 4;
 
-	FVLMVoxelizationParams* VLMVoxelizationParams = GraphBuilder.AllocParameters<FVLMVoxelizationParams>();
-	VLMVoxelizationParams->VolumeCenter = (FVector3f)CubeVolume.GetCenter(); // LWC_TODO: precision loss
-	VLMVoxelizationParams->VolumeExtent = (FVector3f)CubeVolume.GetExtent(); // LWC_TODO: precision loss
-	VLMVoxelizationParams->VolumeMaxDim = CubeMaxDim;
-	VLMVoxelizationParams->VoxelizeVolume = VoxelizationVolumeMips[0]->GetUAV();
-	VLMVoxelizationParams->IndirectionTexture = IndirectionTexture->GetUAV();
-	TRDGUniformBufferRef<FVLMVoxelizationParams> PassUniformBuffer = GraphBuilder.CreateUniformBuffer(VLMVoxelizationParams);
+		FBox CubeVolume(VolumeMin, VolumeMin + FVector(FMath::Max3(VolumeSize.X, VolumeSize.Y, VolumeSize.Z)));
+		int32 CubeMaxDim = FMath::Max3(IndirectionTextureDimensions.X, IndirectionTextureDimensions.Y, IndirectionTextureDimensions.Z);
 
-	for (int32 MipLevel = 0; MipLevel < VoxelizationVolumeMips.Num(); MipLevel++)
-	{
-		FClearVolumeCS::FParameters* Parameters = GraphBuilder.AllocParameters<FClearVolumeCS::FParameters>();
-		Parameters->VolumeSize = VoxelizationVolumeMips[MipLevel]->GetDesc().GetSize();
-		Parameters->VoxelizeVolume = VoxelizationVolumeMips[MipLevel]->GetUAV();
-
-		TShaderMapRef<FClearVolumeCS> ComputeShader(GlobalShaderMap);
-		FComputeShaderUtils::AddPass(
-			GraphBuilder,
-			RDG_EVENT_NAME("ClearVolume"),
-			ComputeShader,
-			Parameters,
-			FComputeShaderUtils::GetGroupCount(VoxelizationVolumeMips[MipLevel]->GetDesc().GetSize(), FIntVector(4)));
-	}
+		FRDGTexture* VoxelizationVolumeMipsRDG = GraphBuilder.RegisterExternalTexture(VoxelizationVolumeMips[0]);
+		FRDGTexture* IndirectTextureRDG = GraphBuilder.RegisterExternalTexture(IndirectionTexture);
 
-	for (FBox ImportanceVolume : ImportanceVolumes)
-	{
-		TShaderMapRef<FVoxelizeImportanceVolumeCS> ComputeShader(GlobalShaderMap);
-
-		FVoxelizeImportanceVolumeCS::FParameters* Parameters = GraphBuilder.AllocParameters<FVoxelizeImportanceVolumeCS::FParameters>();
-		Parameters->VolumeSize = VoxelizationVolumeMips[0]->GetDesc().GetSize();
-		Parameters->ImportanceVolumeMin = (FVector3f)ImportanceVolume.Min;
-		Parameters->ImportanceVolumeMax = (FVector3f)ImportanceVolume.Max;
-		Parameters->VLMVoxelizationParams = PassUniformBuffer;
-		Parameters->VoxelizeVolume = VoxelizationVolumeMips[0]->GetUAV();
-
-		FComputeShaderUtils::AddPass(
-			GraphBuilder,
-			RDG_EVENT_NAME("VoxelizeImportanceVolume"),
-			ComputeShader,
-			Parameters,
-			FComputeShaderUtils::GetGroupCount(VoxelizationVolumeMips[0]->GetDesc().GetSize(), FIntVector(4)));
-	}
+		FVLMVoxelizationParams* VLMVoxelizationParams = GraphBuilder.AllocParameters<FVLMVoxelizationParams>();
+		VLMVoxelizationParams->VolumeCenter = (FVector3f)CubeVolume.GetCenter(); // LWC_TODO: precision loss
+		VLMVoxelizationParams->VolumeExtent = (FVector3f)CubeVolume.GetExtent(); // LWC_TODO: precision loss
+		VLMVoxelizationParams->VolumeMaxDim = CubeMaxDim;
+		VLMVoxelizationParams->VoxelizeVolume = VoxelizationVolumeMipUAVs[0];
+		VLMVoxelizationParams->IndirectionTexture = IndirectionTextureUAV;
+		TRDGUniformBufferRef<FVLMVoxelizationParams> PassUniformBuffer = GraphBuilder.CreateUniformBuffer(VLMVoxelizationParams);
 
-	// Setup ray tracing scene with LOD 0
-	if (!Scene->SetupRayTracingScene())
-	{
-		GraphBuilder.Execute();
-		return;
-	}
+		for (int32 MipLevel = 0; MipLevel < VoxelizationVolumeMips.Num(); MipLevel++)
+		{
+			FClearVolumeCS::FParameters* Parameters = GraphBuilder.AllocParameters<FClearVolumeCS::FParameters>();
+			Parameters->VolumeSize = VoxelizationVolumeMips[MipLevel]->GetDesc().GetSize();
+			Parameters->VoxelizeVolume = VoxelizationVolumeMipUAVs[MipLevel];
+
+			TShaderMapRef<FClearVolumeCS> ComputeShader(GlobalShaderMap);
+			FComputeShaderUtils::AddPass(
+				GraphBuilder,
+				RDG_EVENT_NAME("ClearVolume"),
+				ComputeShader,
+				Parameters,
+				FComputeShaderUtils::GetGroupCount(VoxelizationVolumeMips[MipLevel]->GetDesc().GetSize(), FIntVector(4)));
+		}
 
-	auto* PassParameters = GraphBuilder.AllocParameters<FVoxelizeMeshPassParameters>();
-	PassParameters->PassUniformBuffer = GraphBuilder.CreateUniformBuffer(VLMVoxelizationParams);
-	PassParameters->InstanceCulling = FInstanceCullingContext::CreateDummyInstanceCullingUniformBuffer(GraphBuilder);
+		for (FBox ImportanceVolume : ImportanceVolumes)
+		{
+			TShaderMapRef<FVoxelizeImportanceVolumeCS> ComputeShader(GlobalShaderMap);
+
+			FVoxelizeImportanceVolumeCS::FParameters* Parameters = GraphBuilder.AllocParameters<FVoxelizeImportanceVolumeCS::FParameters>();
+			Parameters->VolumeSize = VoxelizationVolumeMips[0]->GetDesc().GetSize();
+			Parameters->ImportanceVolumeMin = (FVector3f)ImportanceVolume.Min;
+			Parameters->ImportanceVolumeMax = (FVector3f)ImportanceVolume.Max;
+			Parameters->VLMVoxelizationParams = PassUniformBuffer;
+			Parameters->VoxelizeVolume = VoxelizationVolumeMipUAVs[0];
+
+			FComputeShaderUtils::AddPass(
+				GraphBuilder,
+				RDG_EVENT_NAME("VoxelizeImportanceVolume"),
+				ComputeShader,
+				Parameters,
+				FComputeShaderUtils::GetGroupCount(VoxelizationVolumeMips[0]->GetDesc().GetSize(), FIntVector(4)));
+		}
 
-	GraphBuilder.AddPass(
-		RDG_EVENT_NAME("VLM Mesh Voxelization"),
-		PassParameters,
-		ERDGPassFlags::Raster | ERDGPassFlags::NeverCull,
-		[this, CubeMaxDim, PassUniformBuffer, PassParameters](FRHICommandList& RHICmdList)
-	{
-		RHICmdList.SetViewport(0, 0, 0, CubeMaxDim, CubeMaxDim, 1);
+		// Setup ray tracing scene with LOD 0
+		if (!Scene->SetupRayTracingScene())
+		{
+			return;
+		}
 
-		SCOPED_DRAW_EVENTF(RHICmdList, GPULightmassVoxelizeScene, TEXT("GPULightmass VoxelizeScene"));
+		auto* PassParameters = GraphBuilder.AllocParameters<FVoxelizeMeshPassParameters>();
+		PassParameters->PassUniformBuffer = GraphBuilder.CreateUniformBuffer(VLMVoxelizationParams);
+		PassParameters->InstanceCulling = FInstanceCullingContext::CreateDummyInstanceCullingUniformBuffer(GraphBuilder);
 
-		DrawDynamicMeshPass(
-			*Scene->ReferenceView,
-			RHICmdList,
-			[&Scene = Scene, View = Scene->ReferenceView.Get(), ImportanceVolumes = ImportanceVolumes](FDynamicPassMeshDrawListContext* DynamicMeshPassContext)
+		GraphBuilder.AddPass(
+			RDG_EVENT_NAME("VLM Mesh Voxelization"),
+			PassParameters,
+			ERDGPassFlags::Raster | ERDGPassFlags::NeverCull,
+			[this, CubeMaxDim, PassUniformBuffer, PassParameters](FRHICommandList& RHICmdList)
 		{
-			FVLMVoxelizationMeshProcessor MeshProcessor(nullptr, View, DynamicMeshPassContext);
+			RHICmdList.SetViewport(0, 0, 0, CubeMaxDim, CubeMaxDim, 1);
 
-			for (int32 InstanceIndex = 0; InstanceIndex < Scene->StaticMeshInstanceRenderStates.Elements.Num(); InstanceIndex++)
-			{
-				FStaticMeshInstanceRenderState& Instance = Scene->StaticMeshInstanceRenderStates.Elements[InstanceIndex];
+			SCOPED_DRAW_EVENTF(RHICmdList, GPULightmassVoxelizeScene, TEXT("GPULightmass VoxelizeScene"));
 
-				bool bIntersectsAnyImportanceVolume = false;
+			DrawDynamicMeshPass(
+				*Scene->ReferenceView,
+				RHICmdList,
+				[&Scene = Scene, View = Scene->ReferenceView.Get(), ImportanceVolumes = ImportanceVolumes](FDynamicPassMeshDrawListContext* DynamicMeshPassContext)
+			{
+				FVLMVoxelizationMeshProcessor MeshProcessor(nullptr, View, DynamicMeshPassContext);
 
-				for (FBox ImportanceVolume : ImportanceVolumes)
+				for (int32 InstanceIndex = 0; InstanceIndex < Scene->StaticMeshInstanceRenderStates.Elements.Num(); InstanceIndex++)
 				{
-					if (Instance.WorldBounds.GetBox().Intersect(ImportanceVolume))
+					FStaticMeshInstanceRenderState& Instance = Scene->StaticMeshInstanceRenderStates.Elements[InstanceIndex];
+
+					bool bIntersectsAnyImportanceVolume = false;
+
+					for (FBox ImportanceVolume : ImportanceVolumes)
 					{
-						bIntersectsAnyImportanceVolume = true;
-						break;
+						if (Instance.WorldBounds.GetBox().Intersect(ImportanceVolume))
+						{
+							bIntersectsAnyImportanceVolume = true;
+							break;
+						}
 					}
-				}
 
-				if (!bIntersectsAnyImportanceVolume) continue;
+					if (!bIntersectsAnyImportanceVolume) continue;
 
-				TArray<FMeshBatch> MeshBatches = Instance.GetMeshBatchesForGBufferRendering(0);
+					TArray<FMeshBatch> MeshBatches = Instance.GetMeshBatchesForGBufferRendering(0);
 
-				for (auto& MeshBatch : MeshBatches)
-				{
-					MeshBatch.Elements[0].DynamicPrimitiveIndex = InstanceIndex;
-					MeshProcessor.AddMeshBatch(MeshBatch, ~0ull, nullptr);
-				};
-			}
+					for (auto& MeshBatch : MeshBatches)
+					{
+						MeshBatch.Elements[0].DynamicPrimitiveIndex = InstanceIndex;
+						MeshProcessor.AddMeshBatch(MeshBatch, ~0ull, nullptr);
+					};
+				}
 
-			for (int32 InstanceGroupIndex = 0; InstanceGroupIndex < Scene->InstanceGroupRenderStates.Elements.Num(); InstanceGroupIndex++)
-			{
-				FInstanceGroupRenderState& InstanceGroup = Scene->InstanceGroupRenderStates.Elements[InstanceGroupIndex];
+				for (int32 InstanceGroupIndex = 0; InstanceGroupIndex < Scene->InstanceGroupRenderStates.Elements.Num(); InstanceGroupIndex++)
+				{
+					FInstanceGroupRenderState& InstanceGroup = Scene->InstanceGroupRenderStates.Elements[InstanceGroupIndex];
 
-				bool bIntersectsAnyImportanceVolume = false;
+					bool bIntersectsAnyImportanceVolume = false;
 
-				for (FBox ImportanceVolume : ImportanceVolumes)
-				{
-					if (InstanceGroup.WorldBounds.GetBox().Intersect(ImportanceVolume))
+					for (FBox ImportanceVolume : ImportanceVolumes)
 					{
-						bIntersectsAnyImportanceVolume = true;
-						break;
+						if (InstanceGroup.WorldBounds.GetBox().Intersect(ImportanceVolume))
+						{
+							bIntersectsAnyImportanceVolume = true;
+							break;
+						}
 					}
-				}
 
-				if (!bIntersectsAnyImportanceVolume) continue;
+					if (!bIntersectsAnyImportanceVolume) continue;
 
-				TArray<FMeshBatch> MeshBatches = InstanceGroup.GetMeshBatchesForGBufferRendering(0, FTileVirtualCoordinates{});
+					TArray<FMeshBatch> MeshBatches = InstanceGroup.GetMeshBatchesForGBufferRendering(0, FTileVirtualCoordinates{});
 
-				for (auto& MeshBatch : MeshBatches)
-				{
-					MeshBatch.Elements[0].DynamicPrimitiveIndex = Scene->StaticMeshInstanceRenderStates.Elements.Num() + InstanceGroupIndex;
-					MeshProcessor.AddMeshBatch(MeshBatch, ~0ull, nullptr);
-				};
-			}
+					for (auto& MeshBatch : MeshBatches)
+					{
+						MeshBatch.Elements[0].DynamicPrimitiveIndex = Scene->StaticMeshInstanceRenderStates.Elements.Num() + InstanceGroupIndex;
+						MeshProcessor.AddMeshBatch(MeshBatch, ~0ull, nullptr);
+					};
+				}
 
-			for (int32 LandscapeIndex = 0; LandscapeIndex < Scene->LandscapeRenderStates.Elements.Num(); LandscapeIndex++)
-			{
-				FLandscapeRenderState& Landscape = Scene->LandscapeRenderStates.Elements[LandscapeIndex];
+				for (int32 LandscapeIndex = 0; LandscapeIndex < Scene->LandscapeRenderStates.Elements.Num(); LandscapeIndex++)
+				{
+					FLandscapeRenderState& Landscape = Scene->LandscapeRenderStates.Elements[LandscapeIndex];
 
-				bool bIntersectsAnyImportanceVolume = false;
+					bool bIntersectsAnyImportanceVolume = false;
 
-				for (FBox ImportanceVolume : ImportanceVolumes)
-				{
-					if (Landscape.WorldBounds.GetBox().Intersect(ImportanceVolume))
+					for (FBox ImportanceVolume : ImportanceVolumes)
 					{
-						bIntersectsAnyImportanceVolume = true;
-						break;
+						if (Landscape.WorldBounds.GetBox().Intersect(ImportanceVolume))
+						{
+							bIntersectsAnyImportanceVolume = true;
+							break;
+						}
 					}
-				}
 
-				if (!bIntersectsAnyImportanceVolume) continue;
+					if (!bIntersectsAnyImportanceVolume) continue;
 
-				TArray<FMeshBatch> MeshBatches = Landscape.GetMeshBatchesForGBufferRendering(0);
+					TArray<FMeshBatch> MeshBatches = Landscape.GetMeshBatchesForGBufferRendering(0);
 
-				for (auto& MeshBatch : MeshBatches)
-				{
-					MeshBatch.Elements[0].DynamicPrimitiveIndex = Scene->StaticMeshInstanceRenderStates.Elements.Num() + Scene->InstanceGroupRenderStates.Elements.Num() + LandscapeIndex;
-					MeshProcessor.AddMeshBatch(MeshBatch, ~0ull, nullptr);
-				};
-			}
+					for (auto& MeshBatch : MeshBatches)
+					{
+						MeshBatch.Elements[0].DynamicPrimitiveIndex = Scene->StaticMeshInstanceRenderStates.Elements.Num() + Scene->InstanceGroupRenderStates.Elements.Num() + LandscapeIndex;
+						MeshProcessor.AddMeshBatch(MeshBatch, ~0ull, nullptr);
+					};
+				}
+			});
 		});
-	});
 
-	{
-		TShaderMapRef<FDilateVolumeCS> ComputeShader(GlobalShaderMap);
-
-		FDilateVolumeCS::FParameters* Parameters = GraphBuilder.AllocParameters<FDilateVolumeCS::FParameters>();
-		Parameters->VolumeSize = VoxelizationVolumeMips[0]->GetDesc().GetSize();
-		Parameters->VoxelizeVolume = VoxelizationVolumeMips[0]->GetUAV();
-
-		FComputeShaderUtils::AddPass(
-			GraphBuilder,
-			RDG_EVENT_NAME("DilateVolume"),
-			ComputeShader,
-			Parameters,
-			FComputeShaderUtils::GetGroupCount(VoxelizationVolumeMips[0]->GetDesc().GetSize(), FIntVector(4)));
-	}
+		{
+			TShaderMapRef<FDilateVolumeCS> ComputeShader(GlobalShaderMap);
+
+			FDilateVolumeCS::FParameters* Parameters = GraphBuilder.AllocParameters<FDilateVolumeCS::FParameters>();
+			Parameters->VolumeSize = VoxelizationVolumeMips[0]->GetDesc().GetSize();
+			Parameters->VoxelizeVolume = VoxelizationVolumeMipUAVs[0];
+
+			FComputeShaderUtils::AddPass(
+				GraphBuilder,
+				RDG_EVENT_NAME("DilateVolume"),
+				ComputeShader,
+				Parameters,
+				FComputeShaderUtils::GetGroupCount(VoxelizationVolumeMips[0]->GetDesc().GetSize(), FIntVector(4)));
+		}
 
-	for (int32 MipLevel = 1; MipLevel < VoxelizationVolumeMips.Num(); MipLevel++)
-	{
-		TShaderMapRef<FDownsampleVolumeCS> ComputeShader(GlobalShaderMap);
-
-		FDownsampleVolumeCS::FParameters* Parameters = GraphBuilder.AllocParameters<FDownsampleVolumeCS::FParameters>();
-		Parameters->bIsHighestMip = (MipLevel == VoxelizationVolumeMips.Num() - 1) ? 1 : 0;
-		Parameters->VoxelizeVolume = VoxelizationVolumeMips[MipLevel]->GetUAV();
-		Parameters->VoxelizeVolumePrevMip = VoxelizationVolumeMips[MipLevel - 1]->GetUAV();
-
-		FComputeShaderUtils::AddPass(
-			GraphBuilder,
-			RDG_EVENT_NAME("DownsampleVolume"),
-			ComputeShader,
-			Parameters,
-			VoxelizationVolumeMips[MipLevel]->GetDesc().GetSize());
-	}
+		for (int32 MipLevel = 1; MipLevel < VoxelizationVolumeMips.Num(); MipLevel++)
+		{
+			TShaderMapRef<FDownsampleVolumeCS> ComputeShader(GlobalShaderMap);
+
+			FDownsampleVolumeCS::FParameters* Parameters = GraphBuilder.AllocParameters<FDownsampleVolumeCS::FParameters>();
+			Parameters->bIsHighestMip = (MipLevel == VoxelizationVolumeMips.Num() - 1) ? 1 : 0;
+			Parameters->VoxelizeVolume = VoxelizationVolumeMipUAVs[MipLevel];
+			Parameters->VoxelizeVolumePrevMip = VoxelizationVolumeMipUAVs[MipLevel - 1];
+
+			FComputeShaderUtils::AddPass(
+				GraphBuilder,
+				RDG_EVENT_NAME("DownsampleVolume"),
+				ComputeShader,
+				Parameters,
+				VoxelizationVolumeMips[MipLevel]->GetDesc().GetSize());
+		}
 
-	{
-		TResourceArray<int32> InitialBrickAllocatorParams;
-		InitialBrickAllocatorParams.Add(0);
-		InitialBrickAllocatorParams.Add(0);
-		BrickAllocatorParameters.Initialize(TEXT("VolumetricLightmapBrickAllocatorParameters"), 4, 2, PF_R32_SINT, BUF_UnorderedAccess | BUF_SourceCopy, &InitialBrickAllocatorParams);
+		{
+			TResourceArray<int32> InitialBrickAllocatorParams;
+			InitialBrickAllocatorParams.Add(0);
+			InitialBrickAllocatorParams.Add(0);
+			BrickAllocatorParameters.Initialize(TEXT("VolumetricLightmapBrickAllocatorParameters"), 4, 2, PF_R32_SINT, BUF_UnorderedAccess | BUF_SourceCopy, &InitialBrickAllocatorParams);
 
-		RHICmdList.Transition(FRHITransitionInfo(BrickAllocatorParameters.UAV, ERHIAccess::Unknown, ERHIAccess::UAVCompute));
-	}
+			RHICmdList.Transition(FRHITransitionInfo(BrickAllocatorParameters.UAV, ERHIAccess::Unknown, ERHIAccess::UAVCompute));
+		}
 
-	for (int32 MipLevel = VoxelizationVolumeMips.Num() - 1; MipLevel >= 0; MipLevel--)
-	{
-		TShaderMapRef<FCountNumBricksCS> ComputeShader(GlobalShaderMap);
-
-		FCountNumBricksCS::FParameters* Parameters = GraphBuilder.AllocParameters<FCountNumBricksCS::FParameters>();
-		Parameters->VolumeSize = VoxelizationVolumeMips[MipLevel]->GetDesc().GetSize();
-		Parameters->VoxelizeVolume = VoxelizationVolumeMips[MipLevel]->GetUAV();
-		Parameters->BrickAllocatorParameters = BrickAllocatorParameters.UAV;
-
-		FComputeShaderUtils::AddPass(
-			GraphBuilder,
-			RDG_EVENT_NAME("CountNumBricks"),
-			ComputeShader,
-			Parameters,
-			FComputeShaderUtils::GetGroupCount(VoxelizationVolumeMips[MipLevel]->GetDesc().GetSize(), FIntVector(4)));
+		for (int32 MipLevel = VoxelizationVolumeMips.Num() - 1; MipLevel >= 0; MipLevel--)
+		{
+			TShaderMapRef<FCountNumBricksCS> ComputeShader(GlobalShaderMap);
+
+			FCountNumBricksCS::FParameters* Parameters = GraphBuilder.AllocParameters<FCountNumBricksCS::FParameters>();
+			Parameters->VolumeSize = VoxelizationVolumeMips[MipLevel]->GetDesc().GetSize();
+			Parameters->VoxelizeVolume = VoxelizationVolumeMipUAVs[MipLevel];
+			Parameters->BrickAllocatorParameters = BrickAllocatorParameters.UAV;
+
+			FComputeShaderUtils::AddPass(
+				GraphBuilder,
+				RDG_EVENT_NAME("CountNumBricks"),
+				ComputeShader,
+				Parameters,
+				FComputeShaderUtils::GetGroupCount(VoxelizationVolumeMips[MipLevel]->GetDesc().GetSize(), FIntVector(4)));
+		}
 	}
 
-	GraphBuilder.Execute();
-
-	FRHIGPUBufferReadback NumBricksReadback(TEXT("NumBricksReadback"));
-	NumBricksReadback.EnqueueCopy(RHICmdList, BrickAllocatorParameters.Buffer);
-	RHICmdList.BlockUntilGPUIdle();
-	check(NumBricksReadback.IsReady());
 	{
+		FRHIGPUBufferReadback NumBricksReadback(TEXT("NumBricksReadback"));
+		NumBricksReadback.EnqueueCopy(RHICmdList, BrickAllocatorParameters.Buffer);
+		RHICmdList.BlockUntilGPUIdle();
+		check(NumBricksReadback.IsReady());
+
 		int32* Buffer = (int32*)NumBricksReadback.Lock(8);
 		NumTotalBricks = Buffer[0];
 		UE_LOG(LogGPULightmass, Log, TEXT("Volumetric lightmap NumTotalBricks = %d"), NumTotalBricks);
@@ -439,37 +442,60 @@ void FVolumetricLightmapRenderer::VoxelizeScene()
 
 	VolumetricLightmapData.BrickDataDimensions = BrickLayoutDimensions * 5;
 
-	for (int32 MipLevel = VoxelizationVolumeMips.Num() - 1; MipLevel >= 0; MipLevel--)
 	{
-		TShaderMapRef<FGatherBrickRequestsCS> ComputeShader(GlobalShaderMap);
+		FRDGBuilder GraphBuilder(RHICmdList);
 
-		FGatherBrickRequestsCS::FParameters Parameters;
-		Parameters.VolumeSize = VoxelizationVolumeMips[MipLevel]->GetDesc().GetSize();
-		Parameters.BrickSize = 1 << (MipLevel * BrickSizeLog2);
-		Parameters.VoxelizeVolume = VoxelizationVolumeMips[MipLevel]->GetUAV();
-		Parameters.BrickAllocatorParameters = BrickAllocatorParameters.UAV;
-		Parameters.BrickRequests = BrickRequests.UAV;
+		FRDGTextureUAV* IndirectionTextureUAV = GraphBuilder.CreateUAV(GraphBuilder.RegisterExternalTexture(IndirectionTexture));
 
-		FComputeShaderUtils::Dispatch(RHICmdList, ComputeShader, Parameters, FComputeShaderUtils::GetGroupCount(VoxelizationVolumeMips[MipLevel]->GetDesc().GetSize(), FIntVector(4)));
+		TArray<FRDGTextureUAV*, FRDGArrayAllocator> VoxelizationVolumeMipUAVs;
 
-		RHICmdList.Transition(FRHITransitionInfo(BrickRequests.UAV, ERHIAccess::Unknown, ERHIAccess::UAVCompute));
-	}
-
-	for (int32 MipLevel = VoxelizationVolumeMips.Num() - 1; MipLevel >= 0; MipLevel--)
-	{
-		TShaderMapRef<FSplatVolumeCS> ComputeShader(GlobalShaderMap);
+		for (const TRefCountPtr<IPooledRenderTarget>& Mip : VoxelizationVolumeMips)
+		{
+			VoxelizationVolumeMipUAVs.Emplace(GraphBuilder.CreateUAV(GraphBuilder.RegisterExternalTexture(Mip)));
+		}
 
-		FSplatVolumeCS::FParameters Parameters;
-		Parameters.VolumeSize = IndirectionTextureDimensions;
-		Parameters.BrickSize = 1 << (MipLevel * BrickSizeLog2);
-		Parameters.bIsHighestMip = MipLevel == VoxelizationVolumeMips.Num() - 1;
-		Parameters.VoxelizeVolume = VoxelizationVolumeMips[MipLevel]->GetUAV();
-		Parameters.IndirectionTexture = IndirectionTexture->GetUAV();
-		Parameters.BrickAllocatorParameters = BrickAllocatorParameters.UAV;
+		for (int32 MipLevel = VoxelizationVolumeMips.Num() - 1; MipLevel >= 0; MipLevel--)
+		{
+			TShaderMapRef<FGatherBrickRequestsCS> ComputeShader(GlobalShaderMap);
+
+			FGatherBrickRequestsCS::FParameters* PassParameters = GraphBuilder.AllocParameters<FGatherBrickRequestsCS::FParameters>();
+			PassParameters->VolumeSize = VoxelizationVolumeMips[MipLevel]->GetDesc().GetSize();
+			PassParameters->BrickSize = 1 << (MipLevel * BrickSizeLog2);
+			PassParameters->VoxelizeVolume = VoxelizationVolumeMipUAVs[MipLevel];
+			PassParameters->BrickAllocatorParameters = BrickAllocatorParameters.UAV;
+			PassParameters->BrickRequests = BrickRequests.UAV;
+
+			FComputeShaderUtils::AddPass(
+				GraphBuilder,
+				RDG_EVENT_NAME("GatherBrickRequests"),
+				ComputeShader,
+				PassParameters,
+				FComputeShaderUtils::GetGroupCount(VoxelizationVolumeMips[MipLevel]->GetDesc().GetSize(), FIntVector(4)));
+
+			RHICmdList.Transition(FRHITransitionInfo(BrickRequests.UAV, ERHIAccess::Unknown, ERHIAccess::UAVCompute));
+		}
 
-		FComputeShaderUtils::Dispatch(RHICmdList, ComputeShader, Parameters, FComputeShaderUtils::GetGroupCount(IndirectionTextureDimensions, FIntVector(4)));
+		for (int32 MipLevel = VoxelizationVolumeMips.Num() - 1; MipLevel >= 0; MipLevel--)
+		{
+			TShaderMapRef<FSplatVolumeCS> ComputeShader(GlobalShaderMap);
+
+			FSplatVolumeCS::FParameters* PassParameters = GraphBuilder.AllocParameters<FSplatVolumeCS::FParameters>();
+			PassParameters->VolumeSize = IndirectionTextureDimensions;
+			PassParameters->BrickSize = 1 << (MipLevel * BrickSizeLog2);
+			PassParameters->bIsHighestMip = MipLevel == VoxelizationVolumeMips.Num() - 1;
+			PassParameters->VoxelizeVolume = VoxelizationVolumeMipUAVs[MipLevel];
+			PassParameters->IndirectionTexture = IndirectionTextureUAV;
+			PassParameters->BrickAllocatorParameters = BrickAllocatorParameters.UAV;
+
+			FComputeShaderUtils::AddPass(
+				GraphBuilder,
+				RDG_EVENT_NAME("SplatVolume"),
+				ComputeShader,
+				PassParameters,
+				FComputeShaderUtils::GetGroupCount(IndirectionTextureDimensions, FIntVector(4)));
+		}
 
-		RHICmdList.Transition(FRHITransitionInfo(IndirectionTexture->GetUAV(), ERHIAccess::UAVCompute, ERHIAccess::UAVCompute));
+		GraphBuilder.Execute();
 	}
 
 	Scene->DestroyRayTracingScene();
@@ -693,6 +719,8 @@ void FVolumetricLightmapRenderer::BackgroundTick()
 			}
 
 			{
+				FRDGTextureUAV* IndirectionTextureUAV = GraphBuilder.CreateUAV(GraphBuilder.RegisterExternalTexture(IndirectionTexture));
+
 				// Doing 2 passes no longer makes sense in an amortized setup
 				// for (int32 StitchPass = 0; StitchPass < 2; StitchPass++)
 				{
@@ -704,7 +732,7 @@ void FVolumetricLightmapRenderer::BackgroundTick()
 					PassParameters->FrameNumber = FrameNumber / NumFramesOneRound;
 					PassParameters->NumTotalBricks = NumTotalBricks;
 					PassParameters->BrickBatchOffset = BrickBatchOffset;
-					PassParameters->IndirectionTexture = IndirectionTexture->GetUAV();
+					PassParameters->IndirectionTexture = IndirectionTextureUAV;
 					PassParameters->BrickRequests = BrickRequests.UAV;
 					PassParameters->AmbientVector = AccumulationBrickData.AmbientVector.Texture;
 					PassParameters->OutAmbientVector = VolumetricLightmapData.BrickData.AmbientVector.UAV;
diff --git a/Engine/Plugins/Experimental/GPULightmass/Source/GPULightmass/Private/VolumetricLightmapVoxelization.h b/Engine/Plugins/Experimental/GPULightmass/Source/GPULightmass/Private/VolumetricLightmapVoxelization.h
--- a/Engine/Plugins/Experimental/GPULightmass/Source/GPULightmass/Private/VolumetricLightmapVoxelization.h
+++ b/Engine/Plugins/Experimental/GPULightmass/Source/GPULightmass/Private/VolumetricLightmapVoxelization.h
@@ -14,13 +14,13 @@
 #include "MeshMaterialShader.h"
 #include "MeshPassProcessor.inl"
 
-BEGIN_GLOBAL_SHADER_PARAMETER_STRUCT(FVLMVoxelizationParams, )
+BEGIN_UNIFORM_BUFFER_STRUCT(FVLMVoxelizationParams, )
 	SHADER_PARAMETER(FVector4f, VolumeCenter)
 	SHADER_PARAMETER(FVector4f, VolumeExtent)
 	SHADER_PARAMETER(int32, VolumeMaxDim)
-	SHADER_PARAMETER_UAV(RWTexture3D<uint>, VoxelizeVolume)
-	SHADER_PARAMETER_UAV(RWTexture3D<uint4>, IndirectionTexture)
-END_GLOBAL_SHADER_PARAMETER_STRUCT()
+	SHADER_PARAMETER_RDG_TEXTURE_UAV(RWTexture3D<uint>, VoxelizeVolume)
+	SHADER_PARAMETER_RDG_TEXTURE_UAV(RWTexture3D<uint4>, IndirectionTexture)
+END_UNIFORM_BUFFER_STRUCT()
 
 typedef TUniformBufferRef<FVLMVoxelizationParams> FVLMVoxelizationUniformBufferRef;
 
@@ -186,7 +186,7 @@ class FClearVolumeCS : public FGlobalShader
 
 	BEGIN_SHADER_PARAMETER_STRUCT(FParameters, )
 		SHADER_PARAMETER(FIntVector, VolumeSize)
-		SHADER_PARAMETER_UAV(RWTexture3D<uint>, VoxelizeVolume)
+		SHADER_PARAMETER_RDG_TEXTURE_UAV(RWTexture3D<uint>, VoxelizeVolume)
 	END_SHADER_PARAMETER_STRUCT()
 };
 
@@ -209,7 +209,7 @@ class FVoxelizeImportanceVolumeCS : public FGlobalShader
 		SHADER_PARAMETER(FIntVector, VolumeSize)
 		SHADER_PARAMETER(FVector3f, ImportanceVolumeMin)
 		SHADER_PARAMETER(FVector3f, ImportanceVolumeMax)
-		SHADER_PARAMETER_UAV(RWTexture3D<uint>, VoxelizeVolume)
+		SHADER_PARAMETER_RDG_TEXTURE_UAV(RWTexture3D<uint>, VoxelizeVolume)
 		SHADER_PARAMETER_RDG_UNIFORM_BUFFER(FVLMVoxelizationParams, VLMVoxelizationParams)
 	END_SHADER_PARAMETER_STRUCT()
 };
@@ -231,7 +231,7 @@ class FDilateVolumeCS : public FGlobalShader
 
 	BEGIN_SHADER_PARAMETER_STRUCT(FParameters, )
 		SHADER_PARAMETER(FIntVector, VolumeSize)
-		SHADER_PARAMETER_UAV(RWTexture3D<uint>, VoxelizeVolume)
+		SHADER_PARAMETER_RDG_TEXTURE_UAV(RWTexture3D<uint>, VoxelizeVolume)
 	END_SHADER_PARAMETER_STRUCT()
 };
 
@@ -252,8 +252,8 @@ class FDownsampleVolumeCS : public FGlobalShader
 
 	BEGIN_SHADER_PARAMETER_STRUCT(FParameters, )
 		SHADER_PARAMETER(int32, bIsHighestMip)
-		SHADER_PARAMETER_UAV(RWTexture3D<uint>, VoxelizeVolumePrevMip)
-		SHADER_PARAMETER_UAV(RWTexture3D<uint>, VoxelizeVolume)
+		SHADER_PARAMETER_RDG_TEXTURE_UAV(RWTexture3D<uint>, VoxelizeVolumePrevMip)
+		SHADER_PARAMETER_RDG_TEXTURE_UAV(RWTexture3D<uint>, VoxelizeVolume)
 	END_SHADER_PARAMETER_STRUCT()
 };
 
@@ -274,7 +274,7 @@ class FCountNumBricksCS : public FGlobalShader
 
 	BEGIN_SHADER_PARAMETER_STRUCT(FParameters, )
 		SHADER_PARAMETER(FIntVector, VolumeSize)
-		SHADER_PARAMETER_UAV(RWTexture3D<uint>, VoxelizeVolume)
+		SHADER_PARAMETER_RDG_TEXTURE_UAV(RWTexture3D<uint>, VoxelizeVolume)
 		SHADER_PARAMETER_UAV(RWBuffer<int>, BrickAllocatorParameters)
 	END_SHADER_PARAMETER_STRUCT()
 };
@@ -297,7 +297,7 @@ class FGatherBrickRequestsCS : public FGlobalShader
 	BEGIN_SHADER_PARAMETER_STRUCT(FParameters, )
 		SHADER_PARAMETER(FIntVector, VolumeSize)
 		SHADER_PARAMETER(int32, BrickSize)
-		SHADER_PARAMETER_UAV(RWTexture3D<uint>, VoxelizeVolume)
+		SHADER_PARAMETER_RDG_TEXTURE_UAV(RWTexture3D<uint>, VoxelizeVolume)
 		SHADER_PARAMETER_UAV(RWBuffer<int>, BrickAllocatorParameters)
 		SHADER_PARAMETER_UAV(RWBuffer<uint4>, BrickRequests)
 	END_SHADER_PARAMETER_STRUCT()
@@ -322,8 +322,8 @@ class FSplatVolumeCS : public FGlobalShader
 		SHADER_PARAMETER(FIntVector, VolumeSize)
 		SHADER_PARAMETER(int32, BrickSize)
 		SHADER_PARAMETER(int32, bIsHighestMip)
-		SHADER_PARAMETER_UAV(RWTexture3D<uint>, VoxelizeVolume)
-		SHADER_PARAMETER_UAV(RWTexture3D<uint4>, IndirectionTexture)
+		SHADER_PARAMETER_RDG_TEXTURE_UAV(RWTexture3D<uint>, VoxelizeVolume)
+		SHADER_PARAMETER_RDG_TEXTURE_UAV(RWTexture3D<uint4>, IndirectionTexture)
 		SHADER_PARAMETER_UAV(RWBuffer<int>, BrickAllocatorParameters)
 	END_SHADER_PARAMETER_STRUCT()
 };
@@ -349,7 +349,7 @@ class FStitchBorderCS : public FGlobalShader
 		SHADER_PARAMETER(uint32, FrameNumber)
 		SHADER_PARAMETER(int32, NumTotalBricks)
 		SHADER_PARAMETER(int32, BrickBatchOffset)
-		SHADER_PARAMETER_UAV(RWTexture3D<uint4>, IndirectionTexture)
+		SHADER_PARAMETER_RDG_TEXTURE_UAV(RWTexture3D<uint4>, IndirectionTexture)
 		SHADER_PARAMETER_UAV(RWBuffer<uint4>, BrickRequests)
 		SHADER_PARAMETER_TEXTURE(Texture3D<float4>, AmbientVector)
 		SHADER_PARAMETER_UAV(RWTexture3D<float3>, OutAmbientVector)
diff --git a/Engine/Shaders/Public/ShaderVersion.ush b/Engine/Shaders/Public/ShaderVersion.ush
--- a/Engine/Shaders/Public/ShaderVersion.ush
+++ b/Engine/Shaders/Public/ShaderVersion.ush
@@ -2,4 +2,4 @@
 // Each time the console command is executed it generates a new GUID. As this file is included
 // in Platform.ush (which should be included in any shader) it allows to invalidate the shader DDC.
 // 
-// GUID = FF4213C8-41E0-41E2-A4D2-A09146AC6BB6
+// GUID = 65A1E683-92F7-4EAA-BA09-F696E39FC508
