GPULM: Fix hit shader initialization under SM6

diff --git a/Engine/Plugins/Experimental/GPULightmass/Source/GPULightmass/Private/LightmapRayTracing.h b/Engine/Plugins/Experimental/GPULightmass/Source/GPULightmass/Private/LightmapRayTracing.h
--- a/Engine/Plugins/Experimental/GPULightmass/Source/GPULightmass/Private/LightmapRayTracing.h
+++ b/Engine/Plugins/Experimental/GPULightmass/Source/GPULightmass/Private/LightmapRayTracing.h
@@ -28,7 +28,10 @@
 {
 public:
 	FLightmapRayTracingMeshProcessor(FRayTracingMeshCommandContext* InCommandContext, FMeshPassProcessorRenderState InPassDrawRenderState)
-		: FRayTracingMeshProcessor(InCommandContext, nullptr, nullptr, InPassDrawRenderState, ERayTracingMeshCommandsMode::LIGHTMAP_TRACING) {}
+		: FRayTracingMeshProcessor(InCommandContext, nullptr, nullptr, InPassDrawRenderState, ERayTracingMeshCommandsMode::LIGHTMAP_TRACING)
+	{
+		FeatureLevel = GMaxRHIFeatureLevel;
+	}
 };
 
 class FLightmapPathTracingRGS : public FGlobalShader
