Fix static analyzer warnings.

diff --git a/Engine/Plugins/Experimental/GPULightmass/Source/GPULightmass/Private/VolumetricLightmap.cpp b/Engine/Plugins/Experimental/GPULightmass/Source/GPULightmass/Private/VolumetricLightmap.cpp
--- a/Engine/Plugins/Experimental/GPULightmass/Source/GPULightmass/Private/VolumetricLightmap.cpp
+++ b/Engine/Plugins/Experimental/GPULightmass/Source/GPULightmass/Private/VolumetricLightmap.cpp
@@ -572,6 +572,8 @@ void FVolumetricLightmapRenderer::BackgroundTick()
 					TShaderRef<FVolumetricLightmapPathTracingRGS> RayGenShader = GlobalShaderMap->GetShader<FVolumetricLightmapPathTracingRGS>(PermutationVector);
 
 					FVolumetricLightmapPathTracingRGS::FParameters* PassParameters = GraphBuilder.AllocParameters<FVolumetricLightmapPathTracingRGS::FParameters>();
+					CA_ASSUME(PassParameters);
+
 					PassParameters->FrameNumber = FrameNumber / NumFramesOneRound;
 					PassParameters->VolumeMin = (FVector3f)VolumeMin; // LWC_TODO: precision loss
 					PassParameters->VolumeSize = (FVector3f)VolumeSize; // LWC_TODO: precision loss
diff --git a/Engine/Source/Editor/UMGEditor/Private/WidgetBlueprintCompiler.cpp b/Engine/Source/Editor/UMGEditor/Private/WidgetBlueprintCompiler.cpp
--- a/Engine/Source/Editor/UMGEditor/Private/WidgetBlueprintCompiler.cpp
+++ b/Engine/Source/Editor/UMGEditor/Private/WidgetBlueprintCompiler.cpp
@@ -1068,6 +1068,7 @@ void FWidgetBlueprintCompilerContext::FinishCompilingClass(UClass* Class)
 
 	if (!bIsSkeletonOnly)
 	{
+		CA_ASSUME(BPGClass);
 		if (UUserWidget* UserWidget = Cast<UUserWidget>(BPGClass->GetDefaultObject()))
 		{
 			BPGClass->InitializeFieldNotification(UserWidget);
