Remove unused view uniform buffer from lightmap renderer.

diff --git a/Engine/Plugins/Experimental/GPULightmass/Source/GPULightmass/Private/LightmapRenderer.cpp b/Engine/Plugins/Experimental/GPULightmass/Source/GPULightmass/Private/LightmapRenderer.cpp
--- a/Engine/Plugins/Experimental/GPULightmass/Source/GPULightmass/Private/LightmapRenderer.cpp
+++ b/Engine/Plugins/Experimental/GPULightmass/Source/GPULightmass/Private/LightmapRenderer.cpp
@@ -593,9 +593,6 @@ void FCachedRayTracingSceneData::SetupViewUniformBufferFromSceneRenderState(FSce
 			InstancePayloadDataBufferRHI = RHICreateStructuredBuffer(sizeof(FVector4f), InstancePayloadData.GetResourceDataSize(), BUF_Static | BUF_ShaderResource, CreateInfo);
 			InstancePayloadDataBufferSRV = RHICreateShaderResourceView(InstancePayloadDataBufferRHI);
 		}
-
-		FViewUniformShaderParameters ViewUniformBufferParameters;
-		CachedViewUniformBuffer = TUniformBufferRef<FViewUniformShaderParameters>::CreateUniformBufferImmediate(ViewUniformBufferParameters, UniformBuffer_MultiFrame, EUniformBufferValidation::None);
 	}
 }
 
@@ -794,7 +791,6 @@ bool FSceneRenderState::SetupRayTracingScene(int32 LODIndex)
 
 #if 0 // Debug: verify cached ray tracing scene has up-to-date shader bindings
 	TUniquePtr<FCachedRayTracingSceneData> VerificationRayTracingScene = MakeUnique<FCachedRayTracingSceneData>();
-	VerificationRayTracingScene->CachedViewUniformBuffer = CachedRayTracingScene->CachedViewUniformBuffer;
 	VerificationRayTracingScene->SetupFromSceneRenderState(*this);
 
 	check(CachedRayTracingScene->VisibleRayTracingMeshCommands.Num() == VerificationRayTracingScene->VisibleRayTracingMeshCommands.Num());
@@ -866,8 +862,6 @@ bool FSceneRenderState::SetupRayTracingScene(int32 LODIndex)
 			}
 
 			View.ViewUniformBuffer = TUniformBufferRef<FViewUniformShaderParameters>::CreateUniformBufferImmediate(*View.CachedViewUniformShaderParameters, UniformBuffer_SingleFrame);
-
-			CachedRayTracingScene->CachedViewUniformBuffer.UpdateUniformBufferImmediate(*View.CachedViewUniformShaderParameters);
 		}
 
 		const_cast<TRange<int32>&>(View.DynamicPrimitiveCollector.GetPrimitiveIdRange()) = TRange<int32>(0,
diff --git a/Engine/Plugins/Experimental/GPULightmass/Source/GPULightmass/Private/Scene/Scene.h b/Engine/Plugins/Experimental/GPULightmass/Source/GPULightmass/Private/Scene/Scene.h
--- a/Engine/Plugins/Experimental/GPULightmass/Source/GPULightmass/Private/Scene/Scene.h
+++ b/Engine/Plugins/Experimental/GPULightmass/Source/GPULightmass/Private/Scene/Scene.h
@@ -81,8 +81,6 @@ struct FCachedRayTracingSceneData
 	TArray<TArray<FRayTracingGeometryInstance>> RayTracingGeometryInstancesPerLOD;
 	TArray<TUniquePtr<FMatrix>> OwnedRayTracingInstanceTransforms;
 
-	TUniformBufferRef<FViewUniformShaderParameters> CachedViewUniformBuffer;
-
 	void SetupViewUniformBufferFromSceneRenderState(class FSceneRenderState& Scene);
 	void SetupFromSceneRenderState(class FSceneRenderState& Scene);
 };
