GPU Lightmass MGPU:  Bug fix for incorrect GPU assignments for FGPUBatchedTileRequests, introduced when code cleanup was done to move common code into a function, without preserving the GPU index.

#rnx

diff --git a/Engine/Plugins/Experimental/GPULightmass/Source/GPULightmass/Private/LightmapRenderer.cpp b/Engine/Plugins/Experimental/GPULightmass/Source/GPULightmass/Private/LightmapRenderer.cpp
--- a/Engine/Plugins/Experimental/GPULightmass/Source/GPULightmass/Private/LightmapRenderer.cpp
+++ b/Engine/Plugins/Experimental/GPULightmass/Source/GPULightmass/Private/LightmapRenderer.cpp
@@ -191,12 +191,12 @@ struct FGPUBatchedTileRequests
 		}
 	}
 
-	void Commit()
+	void Commit(uint32 GPUIndex)
 	{
 		if (BatchedTilesDesc.Num() > 0)
 		{
 			FRHIResourceCreateInfo CreateInfo(TEXT("BatchedTilesBuffer"));
-			CreateInfo.GPUMask = FRHIGPUMask::GPU0();
+			CreateInfo.GPUMask = FRHIGPUMask::FromIndex(GPUIndex);
 			CreateInfo.ResourceArray = &BatchedTilesDesc;
 
 			BatchedTilesBuffer = RHICreateStructuredBuffer(sizeof(FGPUTileDescription), BatchedTilesDesc.GetResourceDataSize(), BUF_Dynamic | BUF_ShaderResource, CreateInfo);
@@ -1612,7 +1612,7 @@ void FLightmapRenderer::Finalize(FRDGBuilder& GraphBuilder)
 
 			FGPUBatchedTileRequests GPUBatchedTileRequests;
 			GPUBatchedTileRequests.BuildFromTileDescs(TileUploadRequests, LightmapTilePoolGPU, *ScratchTilePoolGPU);
-			GPUBatchedTileRequests.Commit();
+			GPUBatchedTileRequests.Commit(0);
 
 			IPooledRenderTarget* OutputRenderTargets[4] = { nullptr, nullptr, nullptr, nullptr };
 
@@ -2237,7 +2237,7 @@ void FLightmapRenderer::Finalize(FRDGBuilder& GraphBuilder)
 						}
 
 						GPUBatchedTileRequests.BuildFromTileDescs(TileRequestsThisGPU, LightmapTilePoolGPU, *ScratchTilePoolGPU);
-						GPUBatchedTileRequests.Commit();
+						GPUBatchedTileRequests.Commit(GPUIndex);
 						// Let GraphBuilder references the SRV
 						GraphBuilder.AllocObject<FShaderResourceViewRHIRef>(GPUBatchedTileRequests.BatchedTilesSRV);
 
@@ -2650,7 +2650,7 @@ void FLightmapRenderer::Finalize(FRDGBuilder& GraphBuilder)
 						{
 							GPUBatchedTileRequests.BatchedTilesDesc[TileIndex].RenderPassIndex = LightSampleIndexArray[TileIndex];
 						}
-						GPUBatchedTileRequests.Commit();
+						GPUBatchedTileRequests.Commit(GPUIndex);
 						
 						// Let GraphBuilder references the SRV
 						GraphBuilder.AllocObject<FShaderResourceViewRHIRef>(GPUBatchedTileRequests.BatchedTilesSRV);
@@ -2736,7 +2736,7 @@ void FLightmapRenderer::Finalize(FRDGBuilder& GraphBuilder)
 	{
 		FGPUBatchedTileRequests GPUBatchedTileRequests;
 		GPUBatchedTileRequests.BuildFromTileDescs(PendingTileRequests, LightmapTilePoolGPU, *ScratchTilePoolGPU);
-		GPUBatchedTileRequests.Commit();
+		GPUBatchedTileRequests.Commit(0);
 		// Let GraphBuilder references the SRV
 		GraphBuilder.AllocObject<FShaderResourceViewRHIRef>(GPUBatchedTileRequests.BatchedTilesSRV);
 
@@ -2909,7 +2909,7 @@ void FLightmapRenderer::Finalize(FRDGBuilder& GraphBuilder)
 
 				FGPUBatchedTileRequests GPUBatchedTileRequests;
 				GPUBatchedTileRequests.BuildFromTileDescs(LightmapReadbackGroup.ConvergedTileRequests, LightmapTilePoolGPU, *ScratchTilePoolGPU);
-				GPUBatchedTileRequests.Commit();
+				GPUBatchedTileRequests.Commit(GPUIndex);
 				// Let GraphBuilder references the SRV
 				GraphBuilder.AllocObject<FShaderResourceViewRHIRef>(GPUBatchedTileRequests.BatchedTilesSRV);
 
