GPULM: Introduce ApplyRayBias() from path tracer to improve ray bias issue

diff --git a/Engine/Plugins/Experimental/GPULightmass/Shaders/Private/LightmapPathTracing.usf b/Engine/Plugins/Experimental/GPULightmass/Shaders/Private/LightmapPathTracing.usf
--- a/Engine/Plugins/Experimental/GPULightmass/Shaders/Private/LightmapPathTracing.usf
+++ b/Engine/Plugins/Experimental/GPULightmass/Shaders/Private/LightmapPathTracing.usf
@@ -138,6 +138,23 @@ static const float MaxNormalBias = 0.1;
 RaytracingAccelerationStructure TLAS;
 uint SceneVisibleLightCount;
 
+void ApplyRayBias(inout RayDesc Ray, float HitT, float3 Normal, bool ShadowTerminatorFix = false)
+{
+#if 1
+	// Take maximum of position or hit distance to determine "radius" of hit avoidance
+	const float RefValue = max(max(abs(Ray.Origin.x), abs(Ray.Origin.y)), max(abs(Ray.Origin.z), HitT));
+	const uint UlpRadius = 16; // number of floating point ulps to skip around an intersection
+	const float RelativeOffset = asfloat(asuint(RefValue) + UlpRadius) - RefValue;
+	const float ShadowOffset = 0.1; // 1mm
+	const float BaseOffset = 0.001; // 0.01mm (avoid precision issues very close to the origin where ulps become tiny)
+	// for shadow rays, push out the point a bit more so we can reduce shadow terminator issues from poor tesselations
+	// other rays are less likely to arrive at grazing angles, so don't need the extra offset
+	Ray.Origin += max(ShadowTerminatorFix ? ShadowOffset : BaseOffset, RelativeOffset * 100) * Normal;
+#else
+	ApplyPositionBias(Ray.Origin, Ray.Direction, Normal, MaxNormalBias);
+#endif
+}
+
 void AccumulateRadiance(inout float3 TotalRadiance, float3 PathRadiance)
 {
 	float MaxPathIntensity = 0;
@@ -373,6 +390,8 @@ void PathTracingKernel(
 			// GPULightmass fakes a 'camera' ray that hits the lightmap texel perpendicularly
 			Payload.TranslatedWorldPos = TranslatedWorldPosition;
 			Payload.WorldNormal = ShadingNormal;
+			Payload.WorldGeoNormal = ShadingNormal;
+			Payload.WorldSmoothNormal = ShadingNormal;
 			Payload.Radiance = float3(0, 0, 0);
 			Payload.BaseColor = float3(1, 1, 1);
 			Payload.DiffuseColor = float3(1, 1, 1);
@@ -571,8 +590,8 @@ void PathTracingKernel(
 						LightRay.TMin = 0;
 						LightRay.Direction = LightSample.Direction;
 						LightRay.TMax = LightSample.Distance;
-						ApplyPositionBias(LightRay, SignedPositionBias * Payload.WorldNormal, MaxNormalBias);
-
+						ApplyRayBias(LightRay, Payload.HitT, SignedPositionBias * Payload.WorldGeoNormal);
+					
 						float AvgRoughness = ApproximateCaustics ? GetAverageRoughness(Payload) : 0.0;
 
 						Visibility = TraceTransparentVisibilityRay(LightRay, LaunchIndex, AvgRoughness);
@@ -686,7 +705,7 @@ void PathTracingKernel(
 		Ray.Direction = MaterialSample.Direction;
 		Ray.TMin = 0;
 		Ray.TMax = POSITIVE_INFINITY;
-		ApplyPositionBias(Ray, MaterialSample.PositionBiasSign * Payload.WorldNormal, MaxNormalBias);
+		ApplyRayBias(Ray, Payload.HitT, MaterialSample.PositionBiasSign * Payload.WorldGeoNormal);
 
 		if (ApproximateCaustics)
 		{
@@ -1367,8 +1386,8 @@ void VolumetricLightmapPathTracingMainRG()
 			/* out */ Ray.Direction,
 			/* out */ Ray.TMin,
 			/* out */ Ray.TMax);
-			
-		Ray.Origin += Ray.Direction * 0.01f;
+		
+		ApplyRayBias(Ray, 1.0f, ShadingNormal);
 			
 		if (bIsValidRay)
 		{
@@ -1445,7 +1464,7 @@ void StationaryLightShadowTracingMainRG()
 				/* out */ Ray.TMin,
 				/* out */ Ray.TMax);
 				
-			Ray.Origin += Ray.Direction * 0.01f;
+			ApplyRayBias(Ray, 1.0f, WorldNormal);
 				
 			if (bIsValidSample)
 			{
diff --git a/Engine/Shaders/Private/PathTracing/PathTracingMaterialHitShader.usf b/Engine/Shaders/Private/PathTracing/PathTracingMaterialHitShader.usf
--- a/Engine/Shaders/Private/PathTracing/PathTracingMaterialHitShader.usf
+++ b/Engine/Shaders/Private/PathTracing/PathTracingMaterialHitShader.usf
@@ -340,6 +340,8 @@ RAY_TRACING_ENTRY_CLOSEST_HIT(PathTracingMaterialCHS,
 	}
 	Payload.TranslatedWorldPos = LWCToFloat(LWCAdd(MaterialParameters.AbsoluteWorldPosition, ResolvedView.PreViewTranslation));
 	Payload.WorldNormal = normalize(MaterialParameters.WorldNormal);
+	float GeoNormalSign = MaterialParameters.TwoSidedSign * GetPrimitive_DeterminantSign(MaterialParameters.PrimitiveId);
+	Payload.WorldGeoNormal = GeoNormalSign * GeoNormal;
 	Payload.Radiance = GetMaterialEmissive(PixelMaterialInputs);
 	Payload.BaseColor = BaseColor;
 	Payload.Specular = 0;
