GPULM: Fix Irradiance/SH overwriting each other when only one VT layer is requested

diff --git a/Engine/Plugins/Experimental/GPULightmass/Shaders/Private/LightmapOutput.usf b/Engine/Plugins/Experimental/GPULightmass/Shaders/Private/LightmapOutput.usf
--- a/Engine/Plugins/Experimental/GPULightmass/Shaders/Private/LightmapOutput.usf
+++ b/Engine/Plugins/Experimental/GPULightmass/Shaders/Private/LightmapOutput.usf
@@ -45,17 +45,43 @@ void SelectiveLightmapOutputCS(uint3 DispatchThreadId : SV_DispatchThreadID)
 			SH.L2SHCoefficients = SHDirectionality[TexelIndexInPool] / SampleCount;
 			SH.Correction = SHCorrectionAndStationarySkyLightBentNormal[TexelIndexInPool].x / SampleCount;
 			
-			FinalizeLightmapIrradiance(OutputColor, SH, OutputTileAtlas[LaunchIndex + BatchedTiles[TileIndex].OutputLayer0Position]);
-			FinalizeLightmapSH(OutputColor, SH, OutputTileAtlas[LaunchIndex + BatchedTiles[TileIndex].OutputLayer1Position]);
+			if (DIM_OUTPUT_LAYER == 0)
+			{
+				FinalizeLightmapIrradiance(OutputColor, SH, OutputTileAtlas[LaunchIndex + BatchedTiles[TileIndex].OutputLayer0Position]);
+			}
+			
+			if (DIM_OUTPUT_LAYER == 1)
+			{
+				FinalizeLightmapSH(OutputColor, SH, OutputTileAtlas[LaunchIndex + BatchedTiles[TileIndex].OutputLayer1Position]);
+			}
 		}
 		else
 		{
-			OutputTileAtlas[LaunchIndex + BatchedTiles[TileIndex].OutputLayer0Position] = float4(0, 0, 0, 0);
-			OutputTileAtlas[LaunchIndex + BatchedTiles[TileIndex].OutputLayer1Position] = float4(0, 0, 0, 0);
+			if (DIM_OUTPUT_LAYER == 0)
+			{
+				OutputTileAtlas[LaunchIndex + BatchedTiles[TileIndex].OutputLayer0Position] = float4(0, 0, 0, 0);
+			}
+			
+			if (DIM_OUTPUT_LAYER == 1)
+			{
+				OutputTileAtlas[LaunchIndex + BatchedTiles[TileIndex].OutputLayer1Position] = float4(0, 0, 0, 0);
+			}
+		}
+
+#if 0 // Debug: Overwrite with green to see which tiles got output
+		if (DIM_OUTPUT_LAYER == 0)
+		{
+			OutputTileAtlas[LaunchIndex + BatchedTiles[TileIndex].OutputLayer0Position] = float4(0, 1, 0, 0);
 		}
+		
+		if (DIM_OUTPUT_LAYER == 1)
+		{
+			OutputTileAtlas[LaunchIndex + BatchedTiles[TileIndex].OutputLayer1Position] = float4(0, 0, 0, 1);
+		}
+#endif
 	
 #if DRAW_PROGRESS_BARS
-		if (NumTotalSamples > 0 && BatchedTiles[TileIndex].RenderPassIndex < NumTotalSamples - 1)
+		if (DIM_OUTPUT_LAYER == 0 && NumTotalSamples > 0 && BatchedTiles[TileIndex].RenderPassIndex < NumTotalSamples - 1)
 		{
 			int2 ProgressBarMin = int2(64, 64) / 2 - int2(40, 6) / 2;
 			int2 ProgressBarMax = int2(64, 64) / 2 + int2(40, 6) / 2 - int2(1, 1);
diff --git a/Engine/Plugins/Experimental/GPULightmass/Source/GPULightmass/Private/LightmapRenderer.cpp b/Engine/Plugins/Experimental/GPULightmass/Source/GPULightmass/Private/LightmapRenderer.cpp
--- a/Engine/Plugins/Experimental/GPULightmass/Source/GPULightmass/Private/LightmapRenderer.cpp
+++ b/Engine/Plugins/Experimental/GPULightmass/Source/GPULightmass/Private/LightmapRenderer.cpp
@@ -2753,9 +2753,9 @@ void FLightmapRenderer::Finalize(FRDGBuilder& GraphBuilder)
 			RayTracingResolution.X = GPreviewLightmapPhysicalTileSize * GPUBatchedTileRequests.BatchedTilesDesc.Num();
 			RayTracingResolution.Y = GPreviewLightmapPhysicalTileSize;
 
-			if (OutputRenderTargets[0] != nullptr || OutputRenderTargets[1] != nullptr)
+			if (OutputRenderTargets[0] != nullptr)
 			{
-				FRDGTextureRef RenderTargetTileAtlas = GraphBuilder.RegisterExternalTexture(OutputRenderTargets[0] != nullptr ? OutputRenderTargets[0] : OutputRenderTargets[1], TEXT("GPULightmassRenderTargetTileAtlas0"));
+				FRDGTextureRef RenderTargetTileAtlas = GraphBuilder.RegisterExternalTexture(OutputRenderTargets[0], TEXT("GPULightmassRenderTargetTileAtlas0"));
 
 				FSelectiveLightmapOutputCS::FPermutationDomain PermutationVector;
 				PermutationVector.Set<FSelectiveLightmapOutputCS::FOutputLayerDim>(0);
@@ -2782,6 +2782,35 @@ void FLightmapRenderer::Finalize(FRDGBuilder& GraphBuilder)
 					FComputeShaderUtils::GetGroupCount(RayTracingResolution, FComputeShaderUtils::kGolden2DGroupSize));
 			}
 
+			if (OutputRenderTargets[1] != nullptr)
+			{
+				FRDGTextureRef RenderTargetTileAtlas = GraphBuilder.RegisterExternalTexture(OutputRenderTargets[1], TEXT("GPULightmassRenderTargetTileAtlas1"));
+
+				FSelectiveLightmapOutputCS::FPermutationDomain PermutationVector;
+				PermutationVector.Set<FSelectiveLightmapOutputCS::FOutputLayerDim>(1);
+				PermutationVector.Set<FSelectiveLightmapOutputCS::FDrawProgressBars>(Scene->Settings->bShowProgressBars);
+
+				auto Shader = GlobalShaderMap->GetShader<FSelectiveLightmapOutputCS>(PermutationVector);
+
+				FSelectiveLightmapOutputCS::FParameters* PassParameters = GraphBuilder.AllocParameters<FSelectiveLightmapOutputCS::FParameters>();
+				PassParameters->NumBatchedTiles = GPUBatchedTileRequests.BatchedTilesDesc.Num();
+				PassParameters->NumTotalSamples = NumTotalPassesToRender;
+				PassParameters->NumIrradianceCachePasses = Scene->Settings->bUseIrradianceCaching ? Scene->Settings->IrradianceCacheQuality : 0;
+				PassParameters->NumRayGuidingTrialSamples = Scene->Settings->bUseFirstBounceRayGuiding ? Scene->Settings->FirstBounceRayGuidingTrialSamples : 0;
+				PassParameters->BatchedTiles = GPUBatchedTileRequests.BatchedTilesSRV;
+				PassParameters->OutputTileAtlas = GraphBuilder.CreateUAV(RenderTargetTileAtlas);
+				PassParameters->IrradianceAndSampleCount = GraphBuilder.CreateUAV(IrradianceAndSampleCount);
+				PassParameters->SHDirectionality = GraphBuilder.CreateUAV(SHDirectionality);
+				PassParameters->SHCorrectionAndStationarySkyLightBentNormal = GraphBuilder.CreateUAV(SHCorrectionAndStationarySkyLightBentNormal);
+
+				FComputeShaderUtils::AddPass(
+					GraphBuilder,
+					RDG_EVENT_NAME("SelectiveLightmapOutput 1"),
+					Shader,
+					PassParameters,
+					FComputeShaderUtils::GetGroupCount(RayTracingResolution, FComputeShaderUtils::kGolden2DGroupSize));
+			}
+
 			if (OutputRenderTargets[2] != nullptr)
 			{
 				FRDGTextureRef RenderTargetTileAtlas = GraphBuilder.RegisterExternalTexture(OutputRenderTargets[2], TEXT("GPULightmassRenderTargetTileAtlas2"));
