GPULM: Re-enable per-tile ISM instance culling since we have fixed DrawPrimitiveId indexing

diff --git a/Engine/Plugins/Experimental/GPULightmass/Source/GPULightmass/Private/Scene/InstancedStaticMesh.cpp b/Engine/Plugins/Experimental/GPULightmass/Source/GPULightmass/Private/Scene/InstancedStaticMesh.cpp
--- a/Engine/Plugins/Experimental/GPULightmass/Source/GPULightmass/Private/Scene/InstancedStaticMesh.cpp
+++ b/Engine/Plugins/Experimental/GPULightmass/Source/GPULightmass/Private/Scene/InstancedStaticMesh.cpp
@@ -174,10 +174,10 @@ TArray<FMeshBatch> FInstanceGroupRenderState::GetMeshBatchesForGBufferRendering(
 		{
 			MeshBatch.MaterialRenderProxy = ComponentUObject->GetMaterial(Section.MaterialIndex)->GetRenderProxy();
 
-			if (true || CoordsForCulling.MipLevel == -1)
+			if (CoordsForCulling.MipLevel == -1)
 			{
 				// No culling, should be for ray tracing scene
-				MeshBatchElement.UserIndex = 0;
+				MeshBatchElement.UserIndex = INDEX_NONE;
 				MeshBatchElement.NumInstances = InstancedRenderData->PerInstanceRenderData->InstanceBuffer.GetNumInstances();
 
 				MeshBatches.Add(MeshBatch);
