GPULM: Implement TCompactingObjectPool which resembles moving-compacting garbage collection, providing effecient linear access while keeping some flexibility (especially dynamic polymorphism)

diff --git a/Engine/Plugins/Experimental/GPULightmass/Source/GPULightmass/Private/CompactingObjectPool.h b/Engine/Plugins/Experimental/GPULightmass/Source/GPULightmass/Private/CompactingObjectPool.h
new file mode 100644
--- /dev/null
+++ b/Engine/Plugins/Experimental/GPULightmass/Source/GPULightmass/Private/CompactingObjectPool.h
@@ -0,0 +1,271 @@
+﻿// Copyright Epic Games, Inc. All Rights Reserved.
+
+#pragma once
+
+// A handle to an object in a pool which moves and compacts its content.
+class FCompactingObjectHandleBase
+{
+protected:
+	// To allow objects contaning handles to be moved freely (relocatable outside of the handle tracking system),
+	// we use an extra level of indirection so that handles are referenced by their indices ('address') in the handle allocator array.
+	TSparseArray<int32>* HandleAllocator = nullptr;
+	int32 HandleAddress = INDEX_NONE;
+
+	// Pointer to the corresponding tracking table so that we can unregister the handle from the tracking table when it is destructed.
+	TSet<int32>* HandleTrackingTable = nullptr;
+
+public:
+	// Default constructor which takes the default values defined below
+	FCompactingObjectHandleBase() = default;
+	
+	FCompactingObjectHandleBase(int32 ObjectAddress, TSet<int32>* HandleTrackingTable, TSparseArray<int32>* HandleAllocator)
+		: HandleAllocator(HandleAllocator)
+		, HandleTrackingTable(HandleTrackingTable)
+		
+	{
+		// Allocate space for the 'content' of the handle ('content' is the address to the real object)
+		HandleAddress = HandleAllocator->Add(ObjectAddress);
+		if (ObjectAddress != INDEX_NONE)
+		{
+			// Check for double registration
+			check(!HandleTrackingTable->Contains(HandleAddress));
+			// Register this handle in the tracking table
+			HandleTrackingTable->Add(HandleAddress);
+		}
+	}
+
+	void Unregister()
+	{
+		if (GetObjectAddress() != INDEX_NONE)
+		{
+			HandleTrackingTable->Remove(HandleAddress);
+		}
+
+		// The allocator doesn't recycle the memory immediately so we manually mark as invalid here for easier debugging
+		(*HandleAllocator)[HandleAddress] = INDEX_NONE;
+		
+		HandleAllocator->RemoveAt(HandleAddress);
+	}
+
+	~FCompactingObjectHandleBase()
+	{
+		if (HandleAddress != INDEX_NONE)
+		{
+			Unregister();
+		}
+	}
+
+	FCompactingObjectHandleBase& operator=(const FCompactingObjectHandleBase& Other)
+	{
+		// Destroy this handle and create a new one in place
+		FCompactingObjectHandleBase::~FCompactingObjectHandleBase();
+		new (this) FCompactingObjectHandleBase(Other.GetObjectAddress(), Other.HandleTrackingTable, Other.HandleAllocator);
+		return *this;
+	}
+
+	FCompactingObjectHandleBase(const FCompactingObjectHandleBase& Other)
+	{
+		// Destroy this handle and create a new one in place
+		FCompactingObjectHandleBase::~FCompactingObjectHandleBase();
+		new (this) FCompactingObjectHandleBase(Other.GetObjectAddress(), Other.HandleTrackingTable, Other.HandleAllocator);
+	}
+
+	// This operator checks if two handles point to the same object (including invalid object), not for comparing handle addresses
+	bool operator==(const FCompactingObjectHandleBase& Other) const
+	{
+		return HandleAllocator == Other.HandleAllocator && GetObjectAddress() == Other.GetObjectAddress();
+	}
+
+	int32 GetObjectAddress() const
+	{
+		return (*HandleAllocator)[HandleAddress];
+	}
+
+	bool IsInitialized() const
+	{
+		return HandleAddress != INDEX_NONE && HandleAllocator != nullptr && HandleTrackingTable != nullptr;
+	}
+
+	int32 GetObjectAddressChecked() const
+	{
+		check(IsInitialized());
+		return GetObjectAddress();
+	}
+
+	bool PointsToValidObject() const
+	{
+		return IsInitialized() && GetObjectAddress() != INDEX_NONE;
+	}
+	
+	int32 GetValidObjectAddress() const
+	{
+		check(PointsToValidObject());
+		return GetObjectAddress();
+	}
+};
+
+// A tag to prevent user from putting arbitrary containers into TObjectHandle's ObjectStorage ptr (which is practically void*)
+class FDynamicallyIndexableContainer {};
+
+template<typename ObjectType>
+class TObjectHandle : public FCompactingObjectHandleBase
+{
+	using ObjectAccessorFuncType = ObjectType& (*)(FDynamicallyIndexableContainer* ObjectStorage, int32 ObjectAddress);
+	
+private:
+	FDynamicallyIndexableContainer* ObjectStorage = nullptr;
+	ObjectAccessorFuncType ObjectAccessor = nullptr;
+		
+public:
+	TObjectHandle() = default;
+	TObjectHandle(int32 ObjectAddress, TSet<int32>* HandleTrackingTable, TSparseArray<int32>* HandleAllocator, FDynamicallyIndexableContainer* ObjectStorage, ObjectAccessorFuncType ObjectAccessor)
+	: FCompactingObjectHandleBase(ObjectAddress, HandleTrackingTable, HandleAllocator), ObjectStorage(ObjectStorage), ObjectAccessor(ObjectAccessor) {}
+
+	ObjectType& Get() const { return ObjectAccessor(ObjectStorage, GetValidObjectAddress()); }
+	ObjectType& operator*() const { return ObjectAccessor(ObjectStorage, GetValidObjectAddress()); }
+	ObjectType* operator->() const { return &Get(); }
+	operator ObjectType& () const { return Get(); }
+};
+
+template<typename ObjectType>
+class TCompactingObjectPool : public FDynamicallyIndexableContainer
+{
+protected:
+	TArray<TSet<int32>> HandleTrackingTablePerObject;
+	TSparseArray<int32> HandleAllocator;
+	TArray<ObjectType> Objects;
+	
+public:
+	template <typename... ArgsType>
+	TObjectHandle<ObjectType> Emplace(ArgsType&&... Args)
+	{
+		Objects.Emplace(Forward<ArgsType>(Args)...);
+		HandleTrackingTablePerObject.AddDefaulted();
+		return GetHandleForObject(Objects.Last());
+	}
+
+	// Core function for implementing dynamic polymorphic indexing behavior
+	// An array of Base and an array of Derived are inherently different types: their indexing behavior is totally different
+	// The first one indexes each element at i * sizeof(Base) while the second one does so at i * sizeof(Derived)
+	// To properly access the Base of each element out of an array of Derived, the indexing must be done as Derived but
+	// then the indexed element is accessed as Base
+	// This means the concrete type of the element stored in this container (in other words, the type of the container)
+	// must be known at indexing time, prior to the element being accessed - that means dynamic dispatches happen here
+	// before the vtable of the element is accessed. A side effect of that is we don't even need the object's vtable to achieve dynamic polymorphism
+	template<typename InterfaceType = ObjectType>
+	static InterfaceType& ObjectAccessor(FDynamicallyIndexableContainer* ObjectStorage, int32 ObjectAddress)
+	{
+		// Implicit conversion from ObjectType to InterfaceType
+		// Will only compile when InterfaceType is covariant with ObjectType (eg. base classes, or other convertible types)
+		return static_cast<TCompactingObjectPool<ObjectType>*>(ObjectStorage)->Objects[ObjectAddress];
+	}
+	
+	template<typename InterfaceType = ObjectType>
+	TObjectHandle<InterfaceType> GetHandleForObject(ObjectType& Object)
+	{
+		check(&Object >= Objects.GetData() && &Object < Objects.GetData() + Objects.Num());
+		int32 ObjectAddress = &Object - Objects.GetData();
+		return TObjectHandle<InterfaceType>(ObjectAddress, &HandleTrackingTablePerObject[ObjectAddress], &HandleAllocator, this, ObjectAccessor);
+	}
+	
+	// Move the last object in the pool to the address being removed to achieve compaction
+	void RemoveAt(int32 ObjectAddress, bool bCheckNoActiveHandle = false)
+	{
+		Objects.RangeCheck(ObjectAddress);
+
+		const bool bRemovingLastObject = ObjectAddress == Objects.Num() - 1;
+
+		// Redirect handles pointing to the last object to the address being removed
+		if (!bRemovingLastObject)
+		{
+			TSet<int32>& HandleTrackingTableForTheLastObject = HandleTrackingTablePerObject[Objects.Num() - 1];
+			for (int32 HandleAddress : HandleTrackingTableForTheLastObject)
+			{
+				HandleAllocator[HandleAddress] = ObjectAddress;
+			}
+		}
+	
+		if (bCheckNoActiveHandle)
+		{
+			check(HandleTrackingTablePerObject[ObjectAddress].Num() == 0);
+		}
+
+		TSet<int32>& HandleTrackingTable = HandleTrackingTablePerObject[ObjectAddress];
+		TArray<int32> AddressesOfHandlesToInvalidate = HandleTrackingTable.Array();
+		for (int32 HandleAddress : AddressesOfHandlesToInvalidate)
+		{
+			HandleAllocator[HandleAddress] = INDEX_NONE;
+		}
+
+		// Overwrite the handle tracking table of the object being removed with the one for the last object
+		if (!bRemovingLastObject)
+		{
+			HandleTrackingTablePerObject[ObjectAddress] = HandleTrackingTablePerObject[Objects.Num() - 1];
+		}
+
+		// Release the table that is no longer used
+		HandleTrackingTablePerObject.Pop();
+
+		// Move pooled object over
+		// This requires the object to be movable (well, since the underlying TArray requires elements to be bitwise-relocatable, this is always satisfied)
+		if (!bRemovingLastObject)
+		{
+			new (&Objects[ObjectAddress]) ObjectType(MoveTemp(Objects.Last()));
+		}
+
+		Objects.Pop();
+	}
+	
+	void Remove(TObjectHandle<ObjectType>& Handle)
+	{
+		RemoveAt(Handle.GetObjectAddress());
+	
+		check(HandleAllocator[Handle.GetObjectAddress()] == INDEX_NONE);
+	}
+
+	ObjectType& operator[](int32 Index)
+	{
+		return Objects[Index];
+	}
+
+	const ObjectType& operator[](int32 Index) const
+	{
+		return Objects[Index];
+	}
+
+	int32 Num()
+	{
+		return Objects.Num();
+	}
+
+	struct TIterator
+	{
+		TArray<ObjectType>& ObjectStorage;
+		int32 Index;
+
+		TIterator& operator++()
+		{
+			++Index;
+			return *this;
+		}
+
+		bool operator!=(const TIterator& Other)
+		{
+			// We don't check the list of arrays
+			return Index != Other.Index;
+		}
+
+		ObjectType& operator*() const
+		{
+			return ObjectStorage[Index];
+		}
+		
+		ObjectType& operator->() const
+		{
+			return ObjectStorage[Index];
+		}
+	};
+
+	TIterator begin() { return TIterator {Objects, 0};}
+	TIterator end() { return TIterator {Objects, Objects.Num()};}
+};
diff --git a/Engine/Plugins/Experimental/GPULightmass/Source/GPULightmass/Private/CompactingObjectPoolTest.cpp b/Engine/Plugins/Experimental/GPULightmass/Source/GPULightmass/Private/CompactingObjectPoolTest.cpp
new file mode 100644
--- /dev/null
+++ b/Engine/Plugins/Experimental/GPULightmass/Source/GPULightmass/Private/CompactingObjectPoolTest.cpp
@@ -0,0 +1,83 @@
+﻿// Copyright Epic Games, Inc. All Rights Reserved.
+
+PRAGMA_DISABLE_OPTIMIZATION
+
+#include "CompactingObjectPool.h"
+
+struct Base
+{
+	int a;
+
+	operator int&() { return a; }
+};
+
+struct Derived : Base
+{
+	int b;
+};
+
+void RunCompactingObjectPoolTests()
+{
+	TCompactingObjectPool<int32> MyPoolForInts;
+	MyPoolForInts.Emplace(99);
+	MyPoolForInts.Emplace(100);
+	MyPoolForInts.Emplace(101);
+
+	for (int32& myint : MyPoolForInts)
+	{
+		TObjectHandle<int32> MyHandleToAPooledInt = MyPoolForInts.GetHandleForObject(myint);
+		ensure(*MyHandleToAPooledInt == myint);
+	}
+
+	for (auto It = MyPoolForInts.begin(); It != MyPoolForInts.end(); ++It)
+	{
+		int32 Index = It.Index;
+		int32 myint = *It;
+		int32& myint2 = *It;
+	}
+
+	for (int32 Index = 0; Index < MyPoolForInts.Num(); Index++)
+	{
+		int32 myint = MyPoolForInts[Index];
+		int32& myint2 = MyPoolForInts[Index];
+	}
+
+	TObjectHandle<int32> MyHandleToAPooledInt2 = MyPoolForInts.GetHandleForObject(MyPoolForInts[2]);
+	MyPoolForInts.RemoveAt(1);
+	ensure(MyPoolForInts[1] == 101);
+	ensure(MyHandleToAPooledInt2 == 101);
+
+	TCompactingObjectPool<Base> MyPoolForBase;
+	TCompactingObjectPool<Derived> MyPoolForDerived;
+
+	MyPoolForBase.Emplace(Base {0});
+	MyPoolForBase.Emplace(Base {10});
+	MyPoolForBase.Emplace(Base {100});
+	MyPoolForDerived.Emplace(Derived {1, 2});
+	MyPoolForDerived.Emplace(Derived {10, 20});
+	MyPoolForDerived.Emplace(Derived {100, 200});
+
+	static_cast<int&>(MyPoolForDerived[0]) = 3;
+	
+	TObjectHandle<Base> MyDynamicHandle;
+	MyDynamicHandle = MyPoolForBase.GetHandleForObject(MyPoolForBase[0]);
+	ensure(MyDynamicHandle->a == 0);
+	MyDynamicHandle = MyPoolForDerived.GetHandleForObject<Base>(MyPoolForDerived[0]);
+	ensure(MyDynamicHandle->a == 3);
+	
+	// Generate a handle that manipulates the Base::a part of a Derived. Needs operator int&() to be defined.
+	// Evil.
+	TObjectHandle<int> MyDynamicHandleToIntPart = MyPoolForDerived.GetHandleForObject<int>(MyPoolForDerived[0]);
+	// TObjectHandle<Derived> dhb = MyPoolForBase.GetHandleForObject<Derived>(MyPoolForBase[0]); // Shouldn't compile
+	
+	*MyDynamicHandleToIntPart = 4;
+	ensure(MyDynamicHandle->a == 4);
+
+	MyPoolForBase.RemoveAt(0);
+	MyPoolForDerived.RemoveAt(0);
+	
+	ensure(!MyDynamicHandle.PointsToValidObject());
+	ensure(!MyDynamicHandleToIntPart.PointsToValidObject());
+}
+
+PRAGMA_ENABLE_OPTIMIZATION
diff --git a/Engine/Plugins/Experimental/GPULightmass/Source/GPULightmass/Private/GPULightmassModule.cpp b/Engine/Plugins/Experimental/GPULightmass/Source/GPULightmass/Private/GPULightmassModule.cpp
--- a/Engine/Plugins/Experimental/GPULightmass/Source/GPULightmass/Private/GPULightmassModule.cpp
+++ b/Engine/Plugins/Experimental/GPULightmass/Source/GPULightmass/Private/GPULightmassModule.cpp
@@ -10,6 +10,7 @@
 #include "ShaderCore.h"
 #include "GPULightmass.h"
 #include "SceneInterface.h"
+#include "CompactingObjectPool.h"
 
 #define LOCTEXT_NAMESPACE "StaticLightingSystem"
 
@@ -17,8 +18,16 @@ DEFINE_LOG_CATEGORY(LogGPULightmass);
 
 IMPLEMENT_MODULE( FGPULightmassModule, GPULightmass )
 
+void FGPULightmassModule::RunSelfTests()
+{
+	extern void RunCompactingObjectPoolTests();
+	RunCompactingObjectPoolTests();
+}
+
 void FGPULightmassModule::StartupModule()
 {
+	RunSelfTests();
+	
 	UE_LOG(LogGPULightmass, Log, TEXT("GPULightmass module is loaded"));
 	
 	// Maps virtual shader source directory /Plugin/GPULightmass to the plugin's actual Shaders directory.
diff --git a/Engine/Plugins/Experimental/GPULightmass/Source/GPULightmass/Public/GPULightmassModule.h b/Engine/Plugins/Experimental/GPULightmass/Source/GPULightmass/Public/GPULightmassModule.h
--- a/Engine/Plugins/Experimental/GPULightmass/Source/GPULightmass/Public/GPULightmassModule.h
+++ b/Engine/Plugins/Experimental/GPULightmass/Source/GPULightmass/Public/GPULightmassModule.h
@@ -42,6 +42,8 @@ public:
 class FGPULightmassModule : public IGPULightmassModule, public IStaticLightingSystemImpl
 {
 public:
+	void RunSelfTests();
+	
 	/** IModuleInterface implementation */
 	virtual void StartupModule() override;
 	virtual void ShutdownModule() override;
