Fix non-unity build

diff --git a/Engine/Plugins/Experimental/GPULightmass/Source/GPULightmass/Private/CompactingObjectPoolTest.cpp b/Engine/Plugins/Experimental/GPULightmass/Source/GPULightmass/Private/CompactingObjectPoolTest.cpp
--- a/Engine/Plugins/Experimental/GPULightmass/Source/GPULightmass/Private/CompactingObjectPoolTest.cpp
+++ b/Engine/Plugins/Experimental/GPULightmass/Source/GPULightmass/Private/CompactingObjectPoolTest.cpp
@@ -1,5 +1,7 @@
 ﻿// Copyright Epic Games, Inc. All Rights Reserved.
 
+#include "CoreMinimal.h"
+
 PRAGMA_DISABLE_OPTIMIZATION
 
 #include "CompactingObjectPool.h"
