Fix clang compilation issue

diff --git a/Engine/Plugins/Experimental/GPULightmass/Source/GPULightmass/Private/CompactingObjectPoolTest.cpp b/Engine/Plugins/Experimental/GPULightmass/Source/GPULightmass/Private/CompactingObjectPoolTest.cpp
--- a/Engine/Plugins/Experimental/GPULightmass/Source/GPULightmass/Private/CompactingObjectPoolTest.cpp
+++ b/Engine/Plugins/Experimental/GPULightmass/Source/GPULightmass/Private/CompactingObjectPoolTest.cpp
@@ -55,9 +55,9 @@ void RunCompactingObjectPoolTests()
 	MyPoolForBase.Emplace(Base {0});
 	MyPoolForBase.Emplace(Base {10});
 	MyPoolForBase.Emplace(Base {100});
-	MyPoolForDerived.Emplace(Derived {1, 2});
-	MyPoolForDerived.Emplace(Derived {10, 20});
-	MyPoolForDerived.Emplace(Derived {100, 200});
+	MyPoolForDerived.Emplace(Derived {{1}, 2});
+	MyPoolForDerived.Emplace(Derived {{10}, 20});
+	MyPoolForDerived.Emplace(Derived {{100}, 200});
 
 	static_cast<int&>(MyPoolForDerived[0]) = 3;
 	
