GPULM: Refactor lightmap preview VT creation to allow on-demand creation

diff --git a/Engine/Plugins/Experimental/GPULightmass/Source/GPULightmass/Private/LightmapStorage.cpp b/Engine/Plugins/Experimental/GPULightmass/Source/GPULightmass/Private/LightmapStorage.cpp
--- a/Engine/Plugins/Experimental/GPULightmass/Source/GPULightmass/Private/LightmapStorage.cpp
+++ b/Engine/Plugins/Experimental/GPULightmass/Source/GPULightmass/Private/LightmapStorage.cpp
@@ -7,6 +7,7 @@
 #include "Async/Async.h"
 #include "Scene/Lights.h"
 #include "LightmapRenderer.h"
+#include "LightmapPreviewVirtualTexture.h"
 #include "GPULightmassModule.h"
 #include "Async/ParallelFor.h"
 #include "Compression/OodleDataCompression.h"
@@ -170,34 +171,6 @@ FLightmapRenderState::FLightmapRenderState(Initializer InInitializer, FGeometryI
 	{
 		TileRelevantLightSampleCountStates.AddDefaulted(GetPaddedSizeInTilesAtMipLevel(MipLevel).X * GetPaddedSizeInTilesAtMipLevel(MipLevel).Y);
 	}
-
-	{
-		FPrecomputedLightingUniformParameters Parameters;
-		
-		{
-			Parameters.StaticShadowMapMasks = FVector4f(1, 1, 1, 1);
-			Parameters.InvUniformPenumbraSizes = FVector4f(0, 0, 0, 0);
-			Parameters.ShadowMapCoordinateScaleBias = FVector4f(1, 1, 0, 0);
-
-			const uint32 NumCoef = FMath::Max<uint32>(NUM_HQ_LIGHTMAP_COEF, NUM_LQ_LIGHTMAP_COEF);
-			for (uint32 CoefIndex = 0; CoefIndex < NumCoef; ++CoefIndex)
-			{
-				Parameters.LightMapScale[CoefIndex] = FVector4f(1, 1, 1, 1);
-				Parameters.LightMapAdd[CoefIndex] = FVector4f(0, 0, 0, 0);
-			}
-
-			FMemory::Memzero(Parameters.LightmapVTPackedPageTableUniform);
-
-			for (uint32 LayerIndex = 0u; LayerIndex < 5u; ++LayerIndex)
-			{
-				Parameters.LightmapVTPackedUniform[LayerIndex] = FUintVector4(ForceInitToZero);
-			}
-		}
-
-		Parameters.LightMapCoordinateScaleBias = LightmapCoordinateScaleBias;
-
-		SetPrecomputedLightingBuffer(TUniformBufferRef<FPrecomputedLightingUniformParameters>::CreateUniformBufferImmediate(Parameters, UniformBuffer_MultiFrame));
-	}
 }
 
 bool FLightmapRenderState::IsTileGIConverged(FTileVirtualCoordinates Coords, int32 NumGISamples)
@@ -231,4 +204,44 @@ bool FLightmapRenderState::DoesTileHaveValidCPUData(FTileVirtualCoordinates Coor
 	return RetrieveTileState(Coords).CPURevision == CurrentRevision;
 }
 
+void CreateLightmapPreviewVirtualTexture(FLightmapRenderStateRef LightmapRenderState, ERHIFeatureLevel::Type FeatureLevel, FLightmapRenderer* LightmapRenderer)
+{
+	LightmapRenderState->LightmapPreviewVirtualTexture = new FLightmapPreviewVirtualTexture(LightmapRenderState, LightmapRenderer);
+	LightmapRenderState->ResourceCluster->AllocatedVT = LightmapRenderState->LightmapPreviewVirtualTexture->AllocatedVT;
+	LightmapRenderState->ResourceCluster->InitResource();
+
+	{
+		IAllocatedVirtualTexture* AllocatedVT = LightmapRenderState->LightmapPreviewVirtualTexture->AllocatedVT;
+
+		check(AllocatedVT);
+
+		AllocatedVT->GetPackedPageTableUniform(&LightmapRenderState->LightmapVTPackedPageTableUniform[0]);
+		uint32 NumLightmapVTLayers = AllocatedVT->GetNumTextureLayers();
+		for (uint32 LayerIndex = 0u; LayerIndex < NumLightmapVTLayers; ++LayerIndex)
+		{
+			AllocatedVT->GetPackedUniform(&LightmapRenderState->LightmapVTPackedUniform[LayerIndex], LayerIndex);
+		}
+		for (uint32 LayerIndex = NumLightmapVTLayers; LayerIndex < 5u; ++LayerIndex)
+		{
+			LightmapRenderState->LightmapVTPackedUniform[LayerIndex] = FUintVector4(ForceInitToZero);
+		}
+	}
+				
+	LightmapRenderState->ResourceCluster->OverrideFeatureLevel(FeatureLevel);
+	LightmapRenderState->ResourceCluster->UpdateUniformBuffer_RenderThread();
+
+	LightmapRenderState->SetResourceCluster(nullptr);
+}
+
+void FLightmapRenderState::ReleasePreviewVirtualTexture()
+{
+	ResourceCluster->ReleaseResource();
+
+	if (LightmapPreviewVirtualTexture != nullptr)
+	{
+		FVirtualTextureProducerHandle ProducerHandle = LightmapPreviewVirtualTexture->ProducerHandle;
+		GetRendererModule().ReleaseVirtualTextureProducer(ProducerHandle);
+	}
+}
+
 }
diff --git a/Engine/Plugins/Experimental/GPULightmass/Source/GPULightmass/Private/LightmapStorage.h b/Engine/Plugins/Experimental/GPULightmass/Source/GPULightmass/Private/LightmapStorage.h
--- a/Engine/Plugins/Experimental/GPULightmass/Source/GPULightmass/Private/LightmapStorage.h
+++ b/Engine/Plugins/Experimental/GPULightmass/Source/GPULightmass/Private/LightmapStorage.h
@@ -290,6 +290,8 @@ public:
 		return FLightInteraction::LightMap();
 	}
 
+	void ReleasePreviewVirtualTexture();
+
 private:
 	FIntPoint Size;
 	int32 MaxLevel;
@@ -299,6 +301,8 @@ private:
 
 using FLightmapRenderStateRef = TEntityArray<FLightmapRenderState>::EntityRefType;
 
+void CreateLightmapPreviewVirtualTexture(FLightmapRenderStateRef LightmapRenderState, ERHIFeatureLevel::Type FeatureLevel, class FLightmapRenderer* LightmapRenderer);
+
 }
 
 static uint32 GetTypeHash(const GPULightmass::FLightmapRenderStateRef& Ref)
diff --git a/Engine/Plugins/Experimental/GPULightmass/Source/GPULightmass/Private/Scene/Scene.cpp b/Engine/Plugins/Experimental/GPULightmass/Source/GPULightmass/Private/Scene/Scene.cpp
--- a/Engine/Plugins/Experimental/GPULightmass/Source/GPULightmass/Private/Scene/Scene.cpp
+++ b/Engine/Plugins/Experimental/GPULightmass/Source/GPULightmass/Private/Scene/Scene.cpp
@@ -428,8 +428,7 @@
 			FPrimitiveSceneInfo* PrimitiveSceneInfo = SceneProxy->GetPrimitiveSceneInfo();
 			if (PrimitiveSceneInfo && PrimitiveSceneInfo->IsIndexValid())
 			{
-				PrimitiveSceneInfo->UpdateStaticLightingBuffer();
-				PrimitiveSceneInfo->Scene->GPUScene.AddPrimitiveToUpdate(PrimitiveSceneInfo->GetIndex());
+				PrimitiveSceneInfo->Scene->GPUScene.AddPrimitiveToUpdate(PrimitiveSceneInfo->GetIndex(), EPrimitiveDirtyState::ChangedStaticLighting);
 			}
 		}
 	});
@@ -514,8 +513,7 @@
 			FPrimitiveSceneInfo* PrimitiveSceneInfo = SceneProxy->GetPrimitiveSceneInfo();
 			if (PrimitiveSceneInfo && PrimitiveSceneInfo->IsIndexValid())
 			{
-				PrimitiveSceneInfo->UpdateStaticLightingBuffer();
-				PrimitiveSceneInfo->Scene->GPUScene.AddPrimitiveToUpdate(PrimitiveSceneInfo->GetIndex());
+				PrimitiveSceneInfo->Scene->GPUScene.AddPrimitiveToUpdate(PrimitiveSceneInfo->GetIndex(), EPrimitiveDirtyState::ChangedStaticLighting);
 			}
 		}
 	});
@@ -747,7 +745,6 @@
 	}
 
 	TArray<FLightmapRenderState::Initializer> InstanceLightmapRenderStateInitializers;
-	TArray<FLightmapResourceCluster*> ResourceClusters;
 
 	for (FLightmapRef& Lightmap : Instance->LODLightmaps)
 	{
@@ -760,18 +757,18 @@
 				AddLightToLightmap(Lightmap.GetReference_Unsafe(), DirectionalLight);
 			}
 
+			// Ownership will be transferred to render thread, can be nullptr if not created
 			FLightmapResourceCluster* ResourceCluster = Lightmap->ResourceCluster.Release();
 
 			FLightmapRenderState::Initializer Initializer {
 				Lightmap->Name,
 				Lightmap->Size,
 				FMath::Min((int32)FMath::CeilLogTwo((uint32)FMath::Min(Lightmap->GetPaddedSizeInTiles().X, Lightmap->GetPaddedSizeInTiles().Y)), GPreviewLightmapMipmapMaxLevel),
-				ResourceCluster, // temporarily promote unique ptr to raw ptr to make it copyable
+				ResourceCluster,
 				FVector4f(FVector2f(Lightmap->LightmapObject->CoordinateScale), FVector2f(Lightmap->LightmapObject->CoordinateBias))
 			};
 
 			InstanceLightmapRenderStateInitializers.Add(MoveTemp(Initializer));
-			ResourceClusters.Add(ResourceCluster);
 		}
 		else
 		{
@@ -785,6 +782,7 @@
 
 	ENQUEUE_RENDER_COMMAND(RenderThreadInit)(
 		[
+			FeatureLevel = FeatureLevel,
 			InstanceRenderState = MoveTemp(InstanceRenderState), 
 			InstanceLightmapRenderStateInitializers = MoveTemp(InstanceLightmapRenderStateInitializers),
 			&RenderState = RenderState,
@@ -814,28 +812,8 @@
 			if (Initializer.IsValid())
 			{
 				FLightmapRenderStateRef LightmapRenderState = RenderState.LightmapRenderStates.Emplace(Initializer, RenderState.StaticMeshInstanceRenderStates.CreateGeometryInstanceRef(InstanceRenderStateRef, LODIndex));
-				FLightmapPreviewVirtualTexture* LightmapPreviewVirtualTexture = new FLightmapPreviewVirtualTexture(LightmapRenderState, RenderState.LightmapRenderer.Get());
-				LightmapRenderState->LightmapPreviewVirtualTexture = LightmapPreviewVirtualTexture;
-				LightmapRenderState->ResourceCluster->AllocatedVT = LightmapPreviewVirtualTexture->AllocatedVT;
-				LightmapRenderState->ResourceCluster->InitResource();
+				CreateLightmapPreviewVirtualTexture(LightmapRenderState, FeatureLevel, RenderState.LightmapRenderer.Get());
 
-				{
-					IAllocatedVirtualTexture* AllocatedVT = LightmapPreviewVirtualTexture->AllocatedVT;
-
-					check(AllocatedVT);
-
-					AllocatedVT->GetPackedPageTableUniform(&LightmapRenderState->LightmapVTPackedPageTableUniform[0]);
-					uint32 NumLightmapVTLayers = AllocatedVT->GetNumTextureLayers();
-					for (uint32 LayerIndex = 0u; LayerIndex < NumLightmapVTLayers; ++LayerIndex)
-					{
-						AllocatedVT->GetPackedUniform(&LightmapRenderState->LightmapVTPackedUniform[LayerIndex], LayerIndex);
-					}
-					for (uint32 LayerIndex = NumLightmapVTLayers; LayerIndex < 5u; ++LayerIndex)
-					{
-						LightmapRenderState->LightmapVTPackedUniform[LayerIndex] = FUintVector4(ForceInitToZero);
-					}
-				}
-
 				InstanceRenderStateRef->LODLightmapRenderStates.Emplace(LightmapRenderState);
 
 				for (int32 ElementId : RelevantPointLightsToAddOnRenderThread)
@@ -875,11 +853,6 @@
 
 	bNeedsVoxelization = true;
 
-	for (auto ResourceCluster : ResourceClusters)
-	{
-		ResourceCluster->UpdateUniformBuffer(FeatureLevel);
-	}
-
 	if (InComponent->GetWorld())
 	{
 		InComponent->GetWorld()->SendAllEndOfFrameUpdates();
@@ -913,11 +886,7 @@
 		{
 			if (Lightmap.IsValid())
 			{
-				Lightmap->ResourceCluster->ReleaseResource();
-
-				FVirtualTextureProducerHandle ProducerHandle = Lightmap->LightmapPreviewVirtualTexture->ProducerHandle;
-				GetRendererModule().ReleaseVirtualTextureProducer(ProducerHandle);
-
+				Lightmap->ReleasePreviewVirtualTexture();
 				RenderState.LightmapRenderStates.Remove(Lightmap);
 			}
 		}
@@ -965,7 +934,6 @@
 	Instance->AllocateLightmaps(Lightmaps);
 
 	TArray<FLightmapRenderState::Initializer> InstanceLightmapRenderStateInitializers;
-	TArray<FLightmapResourceCluster*> ResourceClusters;
 
 	for (int32 LODIndex = 0; LODIndex < Instance->LODLightmaps.Num(); LODIndex++)
 	{
@@ -1004,18 +972,18 @@
 				AddLightToLightmap(Lightmap.GetReference_Unsafe(), DirectionalLight);
 			}
 
+			// Ownership will be transferred to render thread, can be nullptr if not created
 			FLightmapResourceCluster* ResourceCluster = Lightmap->ResourceCluster.Release();
 
 			FLightmapRenderState::Initializer Initializer {
 				Lightmap->Name,
 				Lightmap->Size,
 				FMath::Min((int32)FMath::CeilLogTwo((uint32)FMath::Min(Lightmap->GetPaddedSizeInTiles().X, Lightmap->GetPaddedSizeInTiles().Y)), GPreviewLightmapMipmapMaxLevel),
-				ResourceCluster, // temporarily promote unique ptr to raw ptr to make it copyable
+				ResourceCluster,
 				FVector4f(FVector4(Lightmap->LightmapObject->CoordinateScale, Lightmap->LightmapObject->CoordinateBias))
 			};
 
 			InstanceLightmapRenderStateInitializers.Add(Initializer);
-			ResourceClusters.Add(ResourceCluster);
 		}
 		else
 		{
@@ -1046,6 +1014,7 @@
 
 	ENQUEUE_RENDER_COMMAND(RenderThreadInit)(
 		[
+			FeatureLevel = FeatureLevel,
 			InstanceRenderState = MoveTemp(InstanceRenderState),
 			InstanceLightmapRenderStateInitializers = MoveTemp(InstanceLightmapRenderStateInitializers),
 			&RenderState = RenderState,
@@ -1063,27 +1032,8 @@
 			if (Initializer.IsValid())
 			{
 				FLightmapRenderStateRef LightmapRenderState = RenderState.LightmapRenderStates.Emplace(Initializer, RenderState.InstanceGroupRenderStates.CreateGeometryInstanceRef(InstanceRenderStateRef, LODIndex));
-				FLightmapPreviewVirtualTexture* LightmapPreviewVirtualTexture = new FLightmapPreviewVirtualTexture(LightmapRenderState, RenderState.LightmapRenderer.Get());
-				LightmapRenderState->LightmapPreviewVirtualTexture = LightmapPreviewVirtualTexture;
-				LightmapRenderState->ResourceCluster->AllocatedVT = LightmapPreviewVirtualTexture->AllocatedVT;
-				LightmapRenderState->ResourceCluster->InitResource();
+				CreateLightmapPreviewVirtualTexture(LightmapRenderState, FeatureLevel, RenderState.LightmapRenderer.Get());
 
-				{
-					IAllocatedVirtualTexture* AllocatedVT = LightmapPreviewVirtualTexture->AllocatedVT;
-
-					check(AllocatedVT);
-
-					AllocatedVT->GetPackedPageTableUniform(&LightmapRenderState->LightmapVTPackedPageTableUniform[0]);
-					uint32 NumLightmapVTLayers = AllocatedVT->GetNumTextureLayers();
-					for (uint32 LayerIndex = 0u; LayerIndex < NumLightmapVTLayers; ++LayerIndex)
-					{
-						AllocatedVT->GetPackedUniform(&LightmapRenderState->LightmapVTPackedUniform[LayerIndex], LayerIndex);
-					}
-					for (uint32 LayerIndex = NumLightmapVTLayers; LayerIndex < 5u; ++LayerIndex)
-					{
-						LightmapRenderState->LightmapVTPackedUniform[LayerIndex] = FUintVector4(ForceInitToZero);
-					}
-				}
 
 				InstanceRenderStateRef->LODLightmapRenderStates.Emplace(MoveTemp(LightmapRenderState));
 
@@ -1114,11 +1064,6 @@
 	});
 
 	bNeedsVoxelization = true;
-
-	for (auto ResourceCluster : ResourceClusters)
-	{
-		ResourceCluster->UpdateUniformBuffer(FeatureLevel);
-	}
 }
 
 void FScene::RemoveGeometryInstanceFromComponent(UInstancedStaticMeshComponent* InComponent)
@@ -1157,11 +1102,7 @@
 		{
 			if (Lightmap.IsValid())
 			{
-				Lightmap->ResourceCluster->ReleaseResource();
-
-				FVirtualTextureProducerHandle ProducerHandle = Lightmap->LightmapPreviewVirtualTexture->ProducerHandle;
-				GetRendererModule().ReleaseVirtualTextureProducer(ProducerHandle);
-
+				Lightmap->ReleasePreviewVirtualTexture();
 				RenderState.LightmapRenderStates.Remove(Lightmap);
 			}
 		}
@@ -1199,7 +1140,6 @@
 	Instance->AllocateLightmaps(Lightmaps);
 
 	TArray<FLightmapRenderState::Initializer> InstanceLightmapRenderStateInitializers;
-	TArray<FLightmapResourceCluster*> ResourceClusters;
 
 	for (int32 LODIndex = 0; LODIndex < Instance->LODLightmaps.Num(); LODIndex++)
 	{
@@ -1217,18 +1157,18 @@
 				AddLightToLightmap(Lightmap.GetReference_Unsafe(), DirectionalLight);
 			}
 
+			// Ownership will be transferred to render thread, can be nullptr if not created
 			FLightmapResourceCluster* ResourceCluster = Lightmap->ResourceCluster.Release();
 
 			FLightmapRenderState::Initializer Initializer{
 				Lightmap->Name,
 				Lightmap->Size,
 				FMath::Min((int32)FMath::CeilLogTwo((uint32)FMath::Min(Lightmap->GetPaddedSizeInTiles().X, Lightmap->GetPaddedSizeInTiles().Y)), GPreviewLightmapMipmapMaxLevel),
-				ResourceCluster, // temporarily promote unique ptr to raw ptr to make it copyable
+				ResourceCluster,
 				FVector4f(FVector4(Lightmap->LightmapObject->CoordinateScale, Lightmap->LightmapObject->CoordinateBias))
 			};
 
 			InstanceLightmapRenderStateInitializers.Add(Initializer);
-			ResourceClusters.Add(ResourceCluster);
 		}
 		else
 		{
@@ -1408,27 +1348,7 @@
 			if (LightmapInitializer.IsValid())
 			{
 				FLightmapRenderStateRef LightmapRenderState = RenderState.LightmapRenderStates.Emplace(LightmapInitializer, RenderState.LandscapeRenderStates.CreateGeometryInstanceRef(InstanceRenderStateRef, LODIndex));
-				FLightmapPreviewVirtualTexture* LightmapPreviewVirtualTexture = new FLightmapPreviewVirtualTexture(LightmapRenderState, RenderState.LightmapRenderer.Get());
-				LightmapRenderState->LightmapPreviewVirtualTexture = LightmapPreviewVirtualTexture;
-				LightmapRenderState->ResourceCluster->AllocatedVT = LightmapPreviewVirtualTexture->AllocatedVT;
-				LightmapRenderState->ResourceCluster->InitResource();
-
-				{
-					IAllocatedVirtualTexture* AllocatedVT = LightmapPreviewVirtualTexture->AllocatedVT;
-
-					check(AllocatedVT);
-
-					AllocatedVT->GetPackedPageTableUniform(&LightmapRenderState->LightmapVTPackedPageTableUniform[0]);
-					uint32 NumLightmapVTLayers = AllocatedVT->GetNumTextureLayers();
-					for (uint32 LayerIndex = 0u; LayerIndex < NumLightmapVTLayers; ++LayerIndex)
-					{
-						AllocatedVT->GetPackedUniform(&LightmapRenderState->LightmapVTPackedUniform[LayerIndex], LayerIndex);
-					}
-					for (uint32 LayerIndex = NumLightmapVTLayers; LayerIndex < 5u; ++LayerIndex)
-					{
-						LightmapRenderState->LightmapVTPackedUniform[LayerIndex] = FUintVector4(ForceInitToZero);
-					}
-				}
+				CreateLightmapPreviewVirtualTexture(LightmapRenderState, LocalFeatureLevel, RenderState.LightmapRenderer.Get());
 
 				InstanceRenderStateRef->LODLightmapRenderStates.Emplace(MoveTemp(LightmapRenderState));
 
@@ -1460,11 +1380,6 @@
 
 	bNeedsVoxelization = true;
 
-	for (auto ResourceCluster : ResourceClusters)
-	{
-		ResourceCluster->UpdateUniformBuffer(FeatureLevel);
-	}
-
 	if (InComponent->GetWorld())
 	{
 		InComponent->GetWorld()->SendAllEndOfFrameUpdates();
@@ -1536,11 +1451,7 @@
 		{
 			if (Lightmap.IsValid())
 			{
-				Lightmap->ResourceCluster->ReleaseResource();
-
-				FVirtualTextureProducerHandle ProducerHandle = Lightmap->LightmapPreviewVirtualTexture->ProducerHandle;
-				GetRendererModule().ReleaseVirtualTextureProducer(ProducerHandle);
-
+				Lightmap->ReleasePreviewVirtualTexture();
 				RenderState.LightmapRenderStates.Remove(Lightmap);
 			}
 		}
diff --git a/Engine/Source/Runtime/Engine/Public/LightMap.h b/Engine/Source/Runtime/Engine/Public/LightMap.h
--- a/Engine/Source/Runtime/Engine/Public/LightMap.h
+++ b/Engine/Source/Runtime/Engine/Public/LightMap.h
@@ -633,6 +633,8 @@ public:
 	virtual void InitRHI();
 	virtual void ReleaseRHI();
 
+	void OverrideFeatureLevel(const FStaticFeatureLevel InFeatureLevel) { SetFeatureLevel(InFeatureLevel); }
+	
 	void UpdateUniformBuffer(ERHIFeatureLevel::Type InFeatureLevel);
 	void UpdateUniformBuffer_RenderThread();
 
