GPULM: Move PendingGITileRequests off RDG allocation as its destruction is not thread-safe (FLightmapTileRequest contains FLightmapRenderStateRef which is a tracked object handle and unregisters itself from the handle tracker on destruction). We're going to have many more RDG passes now (one per tile, max 128 * GPULM workload factor per frame).

diff --git a/Engine/Plugins/Experimental/GPULightmass/Source/GPULightmass/Private/LightmapRenderer.cpp b/Engine/Plugins/Experimental/GPULightmass/Source/GPULightmass/Private/LightmapRenderer.cpp
--- a/Engine/Plugins/Experimental/GPULightmass/Source/GPULightmass/Private/LightmapRenderer.cpp
+++ b/Engine/Plugins/Experimental/GPULightmass/Source/GPULightmass/Private/LightmapRenderer.cpp
@@ -1382,7 +1382,7 @@ void FLightmapRenderer::RenderMeshBatchesIntoGBuffer(
 	FRHICommandList& RHICmdList,
 	const FSceneView* View,
 	int32 GPUScenePrimitiveId,
-	TArray<FMeshBatch>& MeshBatches, 
+	TArray<FMeshBatch> MeshBatches, // Copy mesh batches over so we don't need to worry about their lifetime
 	FVector4f VirtualTexturePhysicalTileCoordinateScaleAndBias,
 	int32 RenderPassIndex,
 	FIntPoint ScratchTilePoolOffset)
@@ -2083,12 +2083,6 @@ void FLightmapRenderer::Finalize(FRDGBuilder& GraphBuilder)
 	NumSamplesPerFrame = FMath::Max(FMath::Min(NumSamplesPerFrame, NumTotalPassesToRender - 1), 0);
 
 	{
-		auto& PendingGITileRequests = *GraphBuilder.AllocObject<TArray<FLightmapTileRequest>>(
-			PendingTileRequests.FilterByPredicate([NumTotalPassesToRender = NumTotalPassesToRender, MostCommonLODIndex](const FLightmapTileRequest& Tile)
-			{
-				return !Tile.RenderState->IsTileGIConverged(Tile.VirtualCoordinates, NumTotalPassesToRender) && Tile.RenderState->GeometryInstanceRef.LODIndex == MostCommonLODIndex;
-			}));
-
 #if RHI_RAYTRACING
 		FLightmapPathTracingRGS::FParameters* PreviousPassParameters[MAX_NUM_GPUS] = {};
 #endif
@@ -2096,6 +2090,12 @@ void FLightmapRenderer::Finalize(FRDGBuilder& GraphBuilder)
 		// Render GI
 		for (int32 SampleIndex = 0; SampleIndex < NumSamplesPerFrame; SampleIndex++)
 		{
+			TArray<FLightmapTileRequest> PendingGITileRequests =
+				PendingTileRequests.FilterByPredicate([NumTotalPassesToRender = NumTotalPassesToRender, MostCommonLODIndex](const FLightmapTileRequest& Tile)
+					{
+					return !Tile.RenderState->IsTileGIConverged(Tile.VirtualCoordinates, NumTotalPassesToRender) && Tile.RenderState->GeometryInstanceRef.LODIndex == MostCommonLODIndex;
+					});
+			
 			if (PendingGITileRequests.Num() > 0)
 			{
 				const int32 AAvsGIMultiplier = 8;
@@ -2147,61 +2147,69 @@ void FLightmapRenderer::Finalize(FRDGBuilder& GraphBuilder)
 							PassParameters->PassUniformBuffer = PassUniformBuffer;
 							PassParameters->InstanceCulling = InstanceCullingUniformBuffer;
 
-							GraphBuilder.AddPass(
-								RDG_EVENT_NAME("LightmapGBuffer"),
-								PassParameters,
-								ERDGPassFlags::Raster,
-								[this, ReferenceView = Scene->ReferenceView, &PendingGITileRequests, &PendingGIRenderPassIndices, GPUIndex, AAvsGIMultiplier](FRHICommandListImmediate& RHICmdList)
+							for (int32 Index = 0; Index < PendingGITileRequests.Num(); Index++)
 							{
-								for (int32 Index = 0; Index < PendingGITileRequests.Num(); Index++)
-								{
-									const FLightmapTileRequest& Tile = PendingGITileRequests[Index];
+								const FLightmapTileRequest& Tile = PendingGITileRequests[Index];
 
-									if (Tile.RenderState->IsTileGIConverged(Tile.VirtualCoordinates, NumTotalPassesToRender)) continue;
-									uint32 AssignedGPUIndex = (Tile.RenderState->DistributionPrefixSum + Tile.RenderState->RetrieveTileStateIndex(Tile.VirtualCoordinates)) % GNumExplicitGPUsForRendering;
-									if (AssignedGPUIndex != GPUIndex) continue;
+								uint32 AssignedGPUIndex = (Tile.RenderState->DistributionPrefixSum + Tile.RenderState->RetrieveTileStateIndex(Tile.VirtualCoordinates)) % GNumExplicitGPUsForRendering;
+								if (AssignedGPUIndex != GPUIndex) continue;
 
-									RHICmdList.SetViewport(0, 0, 0.0f, GPreviewLightmapPhysicalTileSize, GPreviewLightmapPhysicalTileSize, 1.0f);
+								float ScaleX = Tile.RenderState->GetPaddedSizeInTiles().X * GPreviewLightmapVirtualTileSize * 1.0f / (1 << Tile.VirtualCoordinates.MipLevel) / GPreviewLightmapPhysicalTileSize;
+								float ScaleY = Tile.RenderState->GetPaddedSizeInTiles().Y * GPreviewLightmapVirtualTileSize * 1.0f / (1 << Tile.VirtualCoordinates.MipLevel) / GPreviewLightmapPhysicalTileSize;
+								float BiasX = (1.0f * (-Tile.VirtualCoordinates.Position.X * GPreviewLightmapVirtualTileSize) - (-GPreviewLightmapTileBorderSize)) / GPreviewLightmapPhysicalTileSize;
+								float BiasY = (1.0f * (-Tile.VirtualCoordinates.Position.Y * GPreviewLightmapVirtualTileSize) - (-GPreviewLightmapTileBorderSize)) / GPreviewLightmapPhysicalTileSize;
 
-									float ScaleX = Tile.RenderState->GetPaddedSizeInTiles().X * GPreviewLightmapVirtualTileSize * 1.0f / (1 << Tile.VirtualCoordinates.MipLevel) / GPreviewLightmapPhysicalTileSize;
-									float ScaleY = Tile.RenderState->GetPaddedSizeInTiles().Y * GPreviewLightmapVirtualTileSize * 1.0f / (1 << Tile.VirtualCoordinates.MipLevel) / GPreviewLightmapPhysicalTileSize;
-									float BiasX = (1.0f * (-Tile.VirtualCoordinates.Position.X * GPreviewLightmapVirtualTileSize) - (-GPreviewLightmapTileBorderSize)) / GPreviewLightmapPhysicalTileSize;
-									float BiasY = (1.0f * (-Tile.VirtualCoordinates.Position.Y * GPreviewLightmapVirtualTileSize) - (-GPreviewLightmapTileBorderSize)) / GPreviewLightmapPhysicalTileSize;
+								FVector4f VirtualTexturePhysicalTileCoordinateScaleAndBias = FVector4f(ScaleX, ScaleY, BiasX, BiasY);
 
-									FVector4f VirtualTexturePhysicalTileCoordinateScaleAndBias = FVector4f(ScaleX, ScaleY, BiasX, BiasY);
+								TArray<FMeshBatch> MeshBatches = Tile.RenderState->GeometryInstanceRef.GetMeshBatchesForGBufferRendering(Tile.VirtualCoordinates);
 
-									TArray<FMeshBatch> MeshBatches = Tile.RenderState->GeometryInstanceRef.GetMeshBatchesForGBufferRendering(Tile.VirtualCoordinates);
+								int32 EffectiveRenderPassIndex = PendingGIRenderPassIndices[Index];
 
-									int32 EffectiveRenderPassIndex = PendingGIRenderPassIndices[Index];
-									
-									if(Scene->Settings->bUseIrradianceCaching)
+								if (Scene->Settings->bUseIrradianceCaching)
+								{
+									if (EffectiveRenderPassIndex >= Scene->Settings->IrradianceCacheQuality)
 									{
-										if (EffectiveRenderPassIndex >= Scene->Settings->IrradianceCacheQuality)
+										EffectiveRenderPassIndex -= Scene->Settings->IrradianceCacheQuality;
+
+										if (Scene->Settings->bUseFirstBounceRayGuiding)
 										{
-											EffectiveRenderPassIndex -= Scene->Settings->IrradianceCacheQuality;
-									
-											if(Scene->Settings->bUseFirstBounceRayGuiding)
+											if (EffectiveRenderPassIndex >= Scene->Settings->FirstBounceRayGuidingTrialSamples)
 											{
-												if (EffectiveRenderPassIndex >= Scene->Settings->FirstBounceRayGuidingTrialSamples)
-												{
-													EffectiveRenderPassIndex -= Scene->Settings->FirstBounceRayGuidingTrialSamples;
-												}
+												EffectiveRenderPassIndex -= Scene->Settings->FirstBounceRayGuidingTrialSamples;
 											}
 										}
 									}
-									
-									RenderMeshBatchesIntoGBuffer(
-										RHICmdList,
-										ReferenceView.Get(),
-										Scene->GetPrimitiveIdForGPUScene(Tile.RenderState->GeometryInstanceRef),
+								}
+
+								GraphBuilder.AddPass(
+									RDG_EVENT_NAME("LightmapGBufferTile"),
+									PassParameters,
+									ERDGPassFlags::Raster,
+									[
+										this,
+										ReferenceView = Scene->ReferenceView,
+										PrimitiveId = Scene->GetPrimitiveIdForGPUScene(Tile.RenderState->GeometryInstanceRef),
 										MeshBatches,
 										VirtualTexturePhysicalTileCoordinateScaleAndBias,
-										EffectiveRenderPassIndex / AAvsGIMultiplier,
-										ScratchTilePoolGPU->GetPositionFromLinearAddress(Tile.TileAddressInScratch)* GPreviewLightmapPhysicalTileSize);
-
-									GPrimitiveIdVertexBufferPool.DiscardAll();
-								}
-							});
+										EffectiveRenderPassIndex,
+										AAvsGIMultiplier,
+										ScratchTilePoolOffset = ScratchTilePoolGPU->GetPositionFromLinearAddress(Tile.TileAddressInScratch) * GPreviewLightmapPhysicalTileSize
+										](FRHICommandListImmediate& RHICmdList)
+									{
+										RHICmdList.SetViewport(0, 0, 0.0f, GPreviewLightmapPhysicalTileSize, GPreviewLightmapPhysicalTileSize, 1.0f);
+										
+										RenderMeshBatchesIntoGBuffer(
+											RHICmdList,
+											ReferenceView.Get(),
+											PrimitiveId,
+											MeshBatches,
+											VirtualTexturePhysicalTileCoordinateScaleAndBias,
+											EffectiveRenderPassIndex / AAvsGIMultiplier,
+											ScratchTilePoolOffset);
+
+										GPrimitiveIdVertexBufferPool.DiscardAll();
+									});
+							}
 						}
 					}
 				}
diff --git a/Engine/Plugins/Experimental/GPULightmass/Source/GPULightmass/Private/LightmapRenderer.h b/Engine/Plugins/Experimental/GPULightmass/Source/GPULightmass/Private/LightmapRenderer.h
--- a/Engine/Plugins/Experimental/GPULightmass/Source/GPULightmass/Private/LightmapRenderer.h
+++ b/Engine/Plugins/Experimental/GPULightmass/Source/GPULightmass/Private/LightmapRenderer.h
@@ -115,7 +115,7 @@ public:
 		FRHICommandList& RHICmdList,
 		const FSceneView* View,
 		int32 GPUScenePrimitiveId,
-		TArray<FMeshBatch>& MeshBatches,
+		TArray<FMeshBatch> MeshBatches,
 		FVector4f VirtualTexturePhysicalTileCoordinateScaleAndBias,
 		int32 RenderPassIndex,
 		FIntPoint ScratchTilePoolOffset);
