GPULM: Better messages for why ray tracing is disabled

diff --git a/Engine/Plugins/Experimental/GPULightmass/Source/GPULightmassEditor/Private/GPULightmassEditorModule.cpp b/Engine/Plugins/Experimental/GPULightmass/Source/GPULightmassEditor/Private/GPULightmassEditorModule.cpp
--- a/Engine/Plugins/Experimental/GPULightmass/Source/GPULightmassEditor/Private/GPULightmassEditorModule.cpp
+++ b/Engine/Plugins/Experimental/GPULightmass/Source/GPULightmassEditor/Private/GPULightmassEditorModule.cpp
@@ -32,6 +32,117 @@ IMPLEMENT_MODULE( FGPULightmassEditorModule, GPULightmassEditor )
 
 FName GPULightmassSettingsTabName = TEXT("GPULightmassSettings");
 
+#if PLATFORM_WINDOWS
+#include "Windows/AllowWindowsPlatformTypes.h"
+#endif
+
+static bool IsRenderDocPresent()
+{
+#if PLATFORM_WINDOWS
+	// Run only once as this func is called each time the UI is rendered
+	static HMODULE RenderDocDLLHandle = GetModuleHandleW(L"renderdoc.dll");
+	return RenderDocDLLHandle != nullptr;
+#endif
+	// TODO: other platforms
+	return false;
+}
+
+static bool IsCurrentRHIRayTracingCapable()
+{
+	switch(RHIGetInterfaceType())
+	{
+	case ERHIInterfaceType::D3D12:
+	case ERHIInterfaceType::Vulkan:
+		return true;
+	default:
+		return false;
+	}
+}
+
+#if PLATFORM_WINDOWS
+#include "Windows/HideWindowsPlatformTypes.h"
+#endif
+
+enum class ERayTracingDisabledReason : int32
+{
+	OK, // Ray tracing is supported and running
+	DISABLED_BY_PROJECT_SETTINGS, // "Ray Tracing" or "Support Compute Skin Cache" is off in project settings
+	DISABLED_BY_TARGET_PLATFORM, // FXXXTargetPlatform::UsesRayTracing() returns false (eg. WindowsTargetPlatform checks for bEnableRayTracing in its ini)
+	INCOMPATIBLE_SHADER_PLATFORM, // An incompatible shader platform is chosen (eg. ES3.1)
+	INCAPABLE_RHI, // The RHI is incapable of ray tracing in general (eg. DX11)
+	INCOMPATIBLE_PLUGIN, // An incompatible plugin is enabled (specifically, RenderDoc)
+	INCAPABLE_HARDWARE, // The video card isn't capable of hardware ray tracing (too old), or supports only software emulated ray tracing (GTX 10 series)
+};
+
+ERayTracingDisabledReason GetRayTracingDisabledReason()
+{
+	if (IsRayTracingEnabled())
+	{
+		return ERayTracingDisabledReason::OK;
+	}
+	
+	static IConsoleVariable* RayTracingCVar = IConsoleManager::Get().FindConsoleVariable(TEXT("r.RayTracing"));
+	if (RayTracingCVar->GetInt() == 0)
+	{
+		return ERayTracingDisabledReason::DISABLED_BY_PROJECT_SETTINGS;
+	}
+
+	extern RENDERCORE_API uint64 GRayTracingPlaformMask;
+	if (!RHISupportsRayTracing(GMaxRHIShaderPlatform))
+	{
+		return ERayTracingDisabledReason::INCOMPATIBLE_SHADER_PLATFORM;
+	}
+	// Shader platform statically supports ray tracing, but disabled by target platform at runtime
+	else if (RHISupportsRayTracing(GMaxRHIShaderPlatform) && (GRayTracingPlaformMask & (1ull << GMaxRHIShaderPlatform)) == 0)
+	{
+		return ERayTracingDisabledReason::DISABLED_BY_TARGET_PLATFORM;
+	}
+	
+	if (!IsCurrentRHIRayTracingCapable())
+	{
+		return ERayTracingDisabledReason::INCAPABLE_RHI;
+	}
+
+	if (IsRenderDocPresent())
+	{
+		return ERayTracingDisabledReason::INCOMPATIBLE_PLUGIN;
+	}
+
+	// TODO: better determination of reasons instead of pouring everything into this bucket 
+	return ERayTracingDisabledReason::INCAPABLE_HARDWARE;
+}
+
+static FText GenerateRayTracingDisabledReasonMessage(ERayTracingDisabledReason Reason)
+{
+	switch (Reason)
+	{
+	case ERayTracingDisabledReason::DISABLED_BY_PROJECT_SETTINGS:
+		return LOCTEXT("GPULightmassHWRayTracingDisabled",
+		               "GPU Lightmass requires 'Support Hardware Ray Tracing' enabled in your project settings.");
+	case ERayTracingDisabledReason::DISABLED_BY_TARGET_PLATFORM:
+	case ERayTracingDisabledReason::INCOMPATIBLE_SHADER_PLATFORM:
+		return LOCTEXT("GPULightmassHWRayTracingDisabledByPlatformSettings",
+		               "GPU Lightmass requires hardware ray tracing which is disabled by some of your project settings (an incompatible shader platform (eg. ES3.1) is enabled and active, or disabled on your current target platform).");
+	case ERayTracingDisabledReason::INCAPABLE_RHI:
+		return LOCTEXT("GPULightmassHWRayTracingDisabledRHI",
+		               "GPU Lightmass requires hardware ray tracing which is not supported by the current RHI.");
+	case ERayTracingDisabledReason::INCOMPATIBLE_PLUGIN:
+		return LOCTEXT("GPULightmassHWRayTracingDisabledPossiblyByRenderDoc",
+		               "GPU Lightmass requires hardware ray tracing support which is disabled by an incompatible plugin, possibly RenderDoc.");
+	case ERayTracingDisabledReason::INCAPABLE_HARDWARE:
+		return LOCTEXT("GPULightmassHWRayTracingDisabledGPU",
+		               "GPU Lightmass requires hardware ray tracing which isn't supported by your GPU.");
+	default:
+		return FText();
+	}
+}
+
+static bool IsStaticLightingAllowed()
+{
+	static const auto AllowStaticLightingVar = IConsoleManager::Get().FindTConsoleVariableDataInt(TEXT("r.AllowStaticLighting"));
+	return (AllowStaticLightingVar && AllowStaticLightingVar->GetValueOnAnyThread() > 0);
+}
+
 void FGPULightmassEditorModule::StartupModule()
 {
 	FLevelEditorModule& LevelEditorModule = FModuleManager::LoadModuleChecked<FLevelEditorModule>(TEXT("LevelEditor"));
@@ -107,7 +218,7 @@ TSharedRef<SDockTab> FGPULightmassEditorModule::SpawnSettingsTab(const FSpawnTab
 					SNew(SButton)
 					.HAlign(HAlign_Center)
 					.ButtonStyle(FAppStyle::Get(), "FlatButton.Success")
-					.IsEnabled(IsRayTracingEnabled() && IsAllowStaticLightingEnabled())
+					.IsEnabled(IsRayTracingEnabled() && IsStaticLightingAllowed())
 					.Visibility_Lambda([](){ return IsRunning() ? EVisibility::Collapsed : EVisibility::Visible; })
 					.OnClicked_Raw(this, &FGPULightmassEditorModule::OnStartClicked)
 					[
@@ -224,14 +335,11 @@ TSharedRef<SDockTab> FGPULightmassEditorModule::SpawnSettingsTab(const FSpawnTab
 				.AutoWrapText(true)
 				.Text_Lambda( []() -> FText
 				{
-
-					bool bIsRayTracingEnabled = IsRayTracingEnabled();
-					FText RTDisabledMsg = LOCTEXT("GPULightmassRayTracingDisabled", "GPU Lightmass requires ray tracing support which is disabled.");
-					if (!bIsRayTracingEnabled)
+					if (!IsRayTracingEnabled())
 					{
-						return LOCTEXT("GPULightmassRayTracingDisabled", "GPU Lightmass requires ray tracing support which is disabled.");
+						return GenerateRayTracingDisabledReasonMessage(GetRayTracingDisabledReason());
 					}
-					else if (!IsAllowStaticLightingEnabled())
+					else if (!IsStaticLightingAllowed())
 					{
 						return LOCTEXT("GPULightmassAllowStaticLightingDisabled", "GPU Lightmass requires Allow Static Lighting enabled in the project settings.");
 					}
@@ -322,12 +430,6 @@ bool FGPULightmassEditorModule::IsRunning()
 	return false;
 }
 
-bool FGPULightmassEditorModule::IsAllowStaticLightingEnabled()
-{
-	static const auto AllowStaticLightingVar = IConsoleManager::Get().FindTConsoleVariableDataInt(TEXT("r.AllowStaticLighting"));
-	return (AllowStaticLightingVar && AllowStaticLightingVar->GetValueOnAnyThread() > 0);
-}
-
 FReply FGPULightmassEditorModule::OnStartClicked()
 {
 	if (UWorld* World = GEditor->GetEditorWorldContext().World())
diff --git a/Engine/Plugins/Experimental/GPULightmass/Source/GPULightmassEditor/Private/GPULightmassEditorModule.h b/Engine/Plugins/Experimental/GPULightmass/Source/GPULightmassEditor/Private/GPULightmassEditorModule.h
--- a/Engine/Plugins/Experimental/GPULightmass/Source/GPULightmassEditor/Private/GPULightmassEditorModule.h
+++ b/Engine/Plugins/Experimental/GPULightmass/Source/GPULightmassEditor/Private/GPULightmassEditorModule.h
@@ -36,7 +36,6 @@ public:
 	static bool IsRealtimeOn();
 	static bool IsRunning();
 	static bool IsBakeWhatYouSeeMode();
-	static bool IsAllowStaticLightingEnabled();
 
 	FReply OnStartClicked();
 	FReply OnSaveAndStopClicked();
