GPULM: More RDG related fixes

* Move PendingShadowTileRequests out of RDG allocation
* Set static view uniform buffer for landscape dynamic geometry creation

diff --git a/Engine/Plugins/Experimental/GPULightmass/Source/GPULightmass/Private/LightmapRenderer.cpp b/Engine/Plugins/Experimental/GPULightmass/Source/GPULightmass/Private/LightmapRenderer.cpp
--- a/Engine/Plugins/Experimental/GPULightmass/Source/GPULightmass/Private/LightmapRenderer.cpp
+++ b/Engine/Plugins/Experimental/GPULightmass/Source/GPULightmass/Private/LightmapRenderer.cpp
@@ -981,6 +981,7 @@ bool FSceneRenderState::SetupRayTracingScene(int32 LODIndex)
 							FRHIResourceCreateInfo ScratchBufferCreateInfo(TEXT("RHILandscapeScratchBuffer"));
 							FBufferRHIRef ScratchBuffer = RHICreateBuffer(BLASScratchSize, BUF_StructuredBuffer | BUF_RayTracingScratch, 0, ERHIAccess::UAVCompute, ScratchBufferCreateInfo);
 
+							RHICmdList.SetStaticUniformBuffers({ReferenceView->ViewUniformBuffer});
 							DynamicGeometryCollection.DispatchUpdates(RHICmdList, ScratchBuffer);
 
 							// Landscape VF doesn't really use the vertex buffer in HitGroupSystemParameters
@@ -2392,13 +2393,12 @@ void FLightmapRenderer::Finalize(FRDGBuilder& GraphBuilder)
 				{
 					RDG_GPU_MASK_SCOPE(GraphBuilder, FRHIGPUMask::FromIndex(GPUIndex));
 
-					auto& PendingShadowTileRequests = *GraphBuilder.AllocObject<TArray<FLightmapTileRequest>>(
-						PendingShadowTileRequestsOnAllGPUs.FilterByPredicate(
+					TArray<FLightmapTileRequest> PendingShadowTileRequests = PendingShadowTileRequestsOnAllGPUs.FilterByPredicate(
 						[GPUIndex](const FLightmapTileRequest& Tile)
 					{
 						uint32 AssignedGPUIndex = (Tile.RenderState->DistributionPrefixSum + Tile.RenderState->RetrieveTileStateIndex(Tile.VirtualCoordinates)) % GNumExplicitGPUsForRendering;
 						return AssignedGPUIndex == GPUIndex;
-					}));
+					});
 
 					if (PendingShadowTileRequests.Num() == 0) continue;
 
@@ -2614,39 +2614,47 @@ void FLightmapRenderer::Finalize(FRDGBuilder& GraphBuilder)
 						PassParameters->PassUniformBuffer = PassUniformBuffer;
 						PassParameters->InstanceCulling = InstanceCullingUniformBuffer;
 
-						GraphBuilder.AddPass(
-							RDG_EVENT_NAME("LightmapGBuffer"),
-							PassParameters,
-							ERDGPassFlags::Raster,
-							[this, ReferenceView = Scene->ReferenceView, &PendingShadowTileRequests, &LightSampleIndexArray](FRHICommandListImmediate& RHICmdList)
+						for (int32 TileIndex = 0; TileIndex < PendingShadowTileRequests.Num(); TileIndex++)
 						{
-							for (int32 TileIndex = 0; TileIndex < PendingShadowTileRequests.Num(); TileIndex++)
-							{
-								const FLightmapTileRequest& Tile = PendingShadowTileRequests[TileIndex];
+							const FLightmapTileRequest& Tile = PendingShadowTileRequests[TileIndex];
 
-								RHICmdList.SetViewport(0, 0, 0.0f, GPreviewLightmapPhysicalTileSize, GPreviewLightmapPhysicalTileSize, 1.0f);
+							float ScaleX = Tile.RenderState->GetPaddedSizeInTiles().X * GPreviewLightmapVirtualTileSize * 1.0f / (1 << Tile.VirtualCoordinates.MipLevel) / GPreviewLightmapPhysicalTileSize;
+							float ScaleY = Tile.RenderState->GetPaddedSizeInTiles().Y * GPreviewLightmapVirtualTileSize * 1.0f / (1 << Tile.VirtualCoordinates.MipLevel) / GPreviewLightmapPhysicalTileSize;
+							float BiasX = (1.0f * (-Tile.VirtualCoordinates.Position.X * GPreviewLightmapVirtualTileSize) - (-GPreviewLightmapTileBorderSize)) / GPreviewLightmapPhysicalTileSize;
+							float BiasY = (1.0f * (-Tile.VirtualCoordinates.Position.Y * GPreviewLightmapVirtualTileSize) - (-GPreviewLightmapTileBorderSize)) / GPreviewLightmapPhysicalTileSize;
 
-								float ScaleX = Tile.RenderState->GetPaddedSizeInTiles().X * GPreviewLightmapVirtualTileSize * 1.0f / (1 << Tile.VirtualCoordinates.MipLevel) / GPreviewLightmapPhysicalTileSize;
-								float ScaleY = Tile.RenderState->GetPaddedSizeInTiles().Y * GPreviewLightmapVirtualTileSize * 1.0f / (1 << Tile.VirtualCoordinates.MipLevel) / GPreviewLightmapPhysicalTileSize;
-								float BiasX = (1.0f * (-Tile.VirtualCoordinates.Position.X * GPreviewLightmapVirtualTileSize) - (-GPreviewLightmapTileBorderSize)) / GPreviewLightmapPhysicalTileSize;
-								float BiasY = (1.0f * (-Tile.VirtualCoordinates.Position.Y * GPreviewLightmapVirtualTileSize) - (-GPreviewLightmapTileBorderSize)) / GPreviewLightmapPhysicalTileSize;
+							FVector4f VirtualTexturePhysicalTileCoordinateScaleAndBias = FVector4f(ScaleX, ScaleY, BiasX, BiasY);
 
-								FVector4f VirtualTexturePhysicalTileCoordinateScaleAndBias = FVector4f(ScaleX, ScaleY, BiasX, BiasY);
+							TArray<FMeshBatch> MeshBatches = Tile.RenderState->GeometryInstanceRef.GetMeshBatchesForGBufferRendering(Tile.VirtualCoordinates);
 
-								TArray<FMeshBatch> MeshBatches = Tile.RenderState->GeometryInstanceRef.GetMeshBatchesForGBufferRendering(Tile.VirtualCoordinates);
-
-								RenderMeshBatchesIntoGBuffer(
-									RHICmdList,
-									ReferenceView.Get(),
-									Scene->GetPrimitiveIdForGPUScene(Tile.RenderState->GeometryInstanceRef),
+							GraphBuilder.AddPass(
+								RDG_EVENT_NAME("LightmapGBuffer"),
+								PassParameters,
+								ERDGPassFlags::Raster,
+								[
+									this,
+									ReferenceView = Scene->ReferenceView,
+									PrimitiveId = Scene->GetPrimitiveIdForGPUScene(Tile.RenderState->GeometryInstanceRef),
 									MeshBatches,
 									VirtualTexturePhysicalTileCoordinateScaleAndBias,
-									LightSampleIndexArray[TileIndex],
-									ScratchTilePoolGPU->GetPositionFromLinearAddress(Tile.TileAddressInScratch)* GPreviewLightmapPhysicalTileSize);
+									RenderPassIndex = LightSampleIndexArray[TileIndex],
+									ScratchTilePoolOffset = ScratchTilePoolGPU->GetPositionFromLinearAddress(Tile.TileAddressInScratch) * GPreviewLightmapPhysicalTileSize
+									](FRHICommandListImmediate& RHICmdList)
+								{
+									RHICmdList.SetViewport(0, 0, 0.0f, GPreviewLightmapPhysicalTileSize, GPreviewLightmapPhysicalTileSize, 1.0f);
+									
+									RenderMeshBatchesIntoGBuffer(
+										RHICmdList,
+										ReferenceView.Get(),
+										PrimitiveId,
+										MeshBatches,
+										VirtualTexturePhysicalTileCoordinateScaleAndBias,
+										RenderPassIndex,
+										ScratchTilePoolOffset);
 
-								GPrimitiveIdVertexBufferPool.DiscardAll();
-							}
-						});
+									GPrimitiveIdVertexBufferPool.DiscardAll();
+								});
+						}
 					}
 
 #if RHI_RAYTRACING
