GPULM: Writes progress bars to both layer 0 and 1 so they won't be hidden when there's no valid sample

diff --git a/Engine/Plugins/Experimental/GPULightmass/Shaders/Private/LightmapOutput.usf b/Engine/Plugins/Experimental/GPULightmass/Shaders/Private/LightmapOutput.usf
--- a/Engine/Plugins/Experimental/GPULightmass/Shaders/Private/LightmapOutput.usf
+++ b/Engine/Plugins/Experimental/GPULightmass/Shaders/Private/LightmapOutput.usf
@@ -81,7 +81,7 @@ void SelectiveLightmapOutputCS(uint3 DispatchThreadId : SV_DispatchThreadID)
 #endif
 	
 #if DRAW_PROGRESS_BARS
-		if (DIM_OUTPUT_LAYER == 0 && NumTotalSamples > 0 && BatchedTiles[TileIndex].RenderPassIndex < NumTotalSamples - 1)
+		if ((DIM_OUTPUT_LAYER == 0 || DIM_OUTPUT_LAYER == 1) && NumTotalSamples > 0 && BatchedTiles[TileIndex].RenderPassIndex < NumTotalSamples - 1)
 		{
 			int2 ProgressBarMin = int2(64, 64) / 2 - int2(40, 6) / 2;
 			int2 ProgressBarMax = int2(64, 64) / 2 + int2(40, 6) / 2 - int2(1, 1);
@@ -90,7 +90,6 @@ void SelectiveLightmapOutputCS(uint3 DispatchThreadId : SV_DispatchThreadID)
 				if (LaunchIndex.x == ProgressBarMin.x || LaunchIndex.y == ProgressBarMin.y || LaunchIndex.x == ProgressBarMax.x || LaunchIndex.y == ProgressBarMax.y
 					|| LaunchIndex.x <= ProgressBarMin.x + 40 * BatchedTiles[TileIndex].RenderPassIndex / (float)NumTotalSamples)
 				{
-					OutputTileAtlas[LaunchIndex + BatchedTiles[TileIndex].OutputLayer1Position] = float4(0, 0, 0, 1);
 					const half LogBlackPoint = 0.01858136;
 					
 					float3 Color = float3(0, 1, 0);
@@ -105,8 +104,16 @@ void SelectiveLightmapOutputCS(uint3 DispatchThreadId : SV_DispatchThreadID)
 						Color = float3(0.5, 0.5, 0);
 					}
 					
-					OutputTileAtlas[LaunchIndex + BatchedTiles[TileIndex].OutputLayer0Position].rgba = float4(
-					Color, log2( 1 + LogBlackPoint ) - (1.0f / 255 - 0.5 / 255));
+					if (DIM_OUTPUT_LAYER == 0)
+					{
+						OutputTileAtlas[LaunchIndex + BatchedTiles[TileIndex].OutputLayer0Position].rgba = float4(
+						Color, log2( 1 + LogBlackPoint ) - (1.0f / 255 - 0.5 / 255));
+					}
+					
+					if (DIM_OUTPUT_LAYER == 1)
+					{
+						OutputTileAtlas[LaunchIndex + BatchedTiles[TileIndex].OutputLayer1Position] = float4(0, 0, 0, 1);
+					}
 				}
 			}
 		}
