Editor Win64 compilation fix

diff --git a/Engine/Plugins/Experimental/GPULightmass/Source/GPULightmass/Private/LightmapRenderer.cpp b/Engine/Plugins/Experimental/GPULightmass/Source/GPULightmass/Private/LightmapRenderer.cpp
--- a/Engine/Plugins/Experimental/GPULightmass/Source/GPULightmass/Private/LightmapRenderer.cpp
+++ b/Engine/Plugins/Experimental/GPULightmass/Source/GPULightmass/Private/LightmapRenderer.cpp
@@ -562,7 +562,7 @@ void FCachedRayTracingSceneData::SetupViewUniformBufferFromSceneRenderState(FSce
 
 			TResourceArray<FVector4f> InstanceSceneDataSOA;
 			InstanceSceneDataSOA.AddZeroed(FInstanceSceneShaderData::GetDataStrideInFloat4s() * InstanceSceneData.Num());
-			for (int32 ArrayIndex = 0; ArrayIndex < FInstanceSceneShaderData::GetDataStrideInFloat4s(); ArrayIndex++)
+			for (uint32 ArrayIndex = 0; ArrayIndex < FInstanceSceneShaderData::GetDataStrideInFloat4s(); ArrayIndex++)
 			{
 				for (int32 DataIndex = 0; DataIndex < InstanceSceneData.Num(); DataIndex++)
 				{
