GPULM: Implement support for textured rect lights

diff --git a/Engine/Plugins/Experimental/GPULightmass/Source/GPULightmass/Private/LightmapRayTracing.h b/Engine/Plugins/Experimental/GPULightmass/Source/GPULightmass/Private/LightmapRayTracing.h
--- a/Engine/Plugins/Experimental/GPULightmass/Source/GPULightmass/Private/LightmapRayTracing.h
+++ b/Engine/Plugins/Experimental/GPULightmass/Source/GPULightmass/Private/LightmapRayTracing.h
@@ -53,6 +53,7 @@
 	{
 		OutEnvironment.SetDefine(TEXT("GPreviewLightmapPhysicalTileSize"), GPreviewLightmapPhysicalTileSize);
 		OutEnvironment.SetDefine(TEXT("LIGHTMAP_PATH_TRACING_MAIN_RG"), 1);
+		OutEnvironment.SetDefine(TEXT("USE_RECT_LIGHT_TEXTURES"), 1);
 		OutEnvironment.CompilerFlags.Add(CFLAG_ForceDXC);
 		OutEnvironment.CompilerFlags.Add(CFLAG_WarningsAsErrors);
 	}
@@ -111,6 +112,7 @@
 	{
 		OutEnvironment.SetDefine(TEXT("GPreviewLightmapPhysicalTileSize"), GPreviewLightmapPhysicalTileSize);
 		OutEnvironment.SetDefine(TEXT("VOLUMETRIC_LIGHTMAP_PATH_TRACING_MAIN_RG"), 1);
+		OutEnvironment.SetDefine(TEXT("USE_RECT_LIGHT_TEXTURES"), 1);
 		OutEnvironment.CompilerFlags.Add(CFLAG_ForceDXC);
 		OutEnvironment.CompilerFlags.Add(CFLAG_WarningsAsErrors);
 	}
diff --git a/Engine/Plugins/Experimental/GPULightmass/Source/GPULightmass/Private/LightmapRenderer.cpp b/Engine/Plugins/Experimental/GPULightmass/Source/GPULightmass/Private/LightmapRenderer.cpp
--- a/Engine/Plugins/Experimental/GPULightmass/Source/GPULightmass/Private/LightmapRenderer.cpp
+++ b/Engine/Plugins/Experimental/GPULightmass/Source/GPULightmass/Private/LightmapRenderer.cpp
@@ -27,6 +27,7 @@
 #include "ScreenPass.h"
 #include "RayTracingDynamicGeometryCollection.h"
 #include "ShaderCompiler.h"
+#include "RectLightTextureManager.h"
 
 class FCopyConvergedLightmapTilesCS : public FGlobalShader
 {
@@ -1930,6 +1931,8 @@ void FLightmapRenderer::Finalize(FRDGBuilder& GraphBuilder)
 		MostCommonLODIndex = NonZeroLODIndices[FrameNumber % NonZeroLODIndices.Num()];
 	}
 
+	RectLightAtlas::UpdateRectLightAtlasTexture(GraphBuilder, Scene->FeatureLevel);
+
 	Scene->SetupRayTracingScene(MostCommonLODIndex);
 
 	TStaticArray<FRDGTextureUAVRef, 3> ScratchTilePoolLayerUAVs;
diff --git a/Engine/Plugins/Experimental/GPULightmass/Source/GPULightmass/Private/PathTracingLightParameters.inl b/Engine/Plugins/Experimental/GPULightmass/Source/GPULightmass/Private/PathTracingLightParameters.inl
--- a/Engine/Plugins/Experimental/GPULightmass/Source/GPULightmass/Private/PathTracingLightParameters.inl
+++ b/Engine/Plugins/Experimental/GPULightmass/Source/GPULightmass/Private/PathTracingLightParameters.inl
@@ -11,6 +11,11 @@ RENDERER_API FRDGTexture* PrepareIESAtlas(const TMap<FTexture*, int>& InIESLight
 
 RENDERER_API void PrepareLightGrid(FRDGBuilder& GraphBuilder, FPathTracingLightGrid* LightGridParameters, const FPathTracingLight* Lights, uint32 NumLights, uint32 NumInfiniteLights, FRDGBufferSRV* LightsSRV);
 
+static uint32 EncodeToF16x2(const FVector2f& In)
+{
+	return FFloat16(In.X).Encoded | (FFloat16(In.Y).Encoded << 16);
+}
+
 template<typename PassParameterType>
 void SetupPathTracingLightParameters(
 	const GPULightmass::FLightSceneRenderState& LightScene,
@@ -205,6 +210,13 @@ void SetupPathTracingLightParameters(
 		FVector3f Tip = Center + Normal * Radius;
 		DestLight.TranslatedBoundMin = Tip.ComponentMin(Center - Radius * Disc);
 		DestLight.TranslatedBoundMax = Tip.ComponentMax(Center + Radius * Disc);
+
+		DestLight.RectLightAtlasUVOffset = EncodeToF16x2(Light.RectLightAtlasUVOffset);
+		DestLight.RectLightAtlasUVScale  = EncodeToF16x2(Light.RectLightAtlasUVScale);
+		if (Light.RectLightAtlasMaxLevel < 16)
+		{
+			DestLight.Flags |= PATHTRACER_FLAG_HAS_RECT_TEXTURE_MASK;
+		}
 	}
 
 	PassParameters->SceneLightCount = Lights.Num();
diff --git a/Engine/Plugins/Experimental/GPULightmass/Source/GPULightmass/Private/Scene/Lights.cpp b/Engine/Plugins/Experimental/GPULightmass/Source/GPULightmass/Private/Scene/Lights.cpp
--- a/Engine/Plugins/Experimental/GPULightmass/Source/GPULightmass/Private/Scene/Lights.cpp
+++ b/Engine/Plugins/Experimental/GPULightmass/Source/GPULightmass/Private/Scene/Lights.cpp
@@ -3,6 +3,7 @@
 #include "Lights.h"
 #include "RenderGraphBuilder.h"
 #include "ReflectionEnvironment.h"
+#include "RectLightTextureManager.h"
 
 RENDERER_API void PrepareSkyTexture_Internal(
 	FRDGBuilder& GraphBuilder,
@@ -200,6 +201,22 @@ FRectLightRenderState::FRectLightRenderState(URectLightComponent* RectLightCompo
 	AttenuationRadius = RectLightComponent->AttenuationRadius;
 	ShadowMapChannel = RectLightComponent->PreviewShadowMapChannel;
 	IESTexture = RectLightComponent->IESTexture ? RectLightComponent->IESTexture->GetResource() : nullptr;
+	SourceTexture = RectLightComponent->SourceTexture;
+}
+
+void FRectLightRenderState::RenderThreadInit()
+{
+	AtlasSlotIndex = RectLightAtlas::AddRectLightTexture(SourceTexture);
+
+	const RectLightAtlas::FAtlasSlotDesc Slot = RectLightAtlas::GetRectLightAtlasSlot(AtlasSlotIndex);
+	RectLightAtlasUVOffset = Slot.UVOffset;
+	RectLightAtlasUVScale = Slot.UVScale;
+	RectLightAtlasMaxLevel = Slot.MaxMipLevel;
+}
+
+void FRectLightRenderState::RenderThreadFinalize()
+{
+	RectLightAtlas::RemoveRectLightTexture(AtlasSlotIndex);
 }
 
 FLightRenderParameters FDirectionalLightRenderState::GetLightShaderParameters() const
@@ -279,10 +296,7 @@ FLightRenderParameters FRectLightRenderState::GetLightShaderParameters() const
 	LightParameters.SourceLength = SourceHeight * 0.5f;
 	LightParameters.RectLightBarnCosAngle = FMath::Cos(FMath::DegreesToRadians(BarnDoorAngle));
 	LightParameters.RectLightBarnLength = BarnDoorLength;
-	LightParameters.RectLightAtlasUVOffset = FVector2f::ZeroVector;
-	LightParameters.RectLightAtlasUVScale = FVector2f::ZeroVector;
-	LightParameters.RectLightAtlasMaxLevel = FLightRenderParameters::GetRectLightAtlasInvalidMIPLevel();
-
+	
 	return LightParameters;
 }
 
diff --git a/Engine/Plugins/Experimental/GPULightmass/Source/GPULightmass/Private/Scene/Lights.h b/Engine/Plugins/Experimental/GPULightmass/Source/GPULightmass/Private/Scene/Lights.h
--- a/Engine/Plugins/Experimental/GPULightmass/Source/GPULightmass/Private/Scene/Lights.h
+++ b/Engine/Plugins/Experimental/GPULightmass/Source/GPULightmass/Private/Scene/Lights.h
@@ -85,6 +85,8 @@ struct FLocalLightBuildInfo
 struct FLocalLightRenderState
 {
 	FLocalLightRenderState(ULightComponent* LightComponent);
+	virtual void RenderThreadInit()  {}
+	virtual void RenderThreadFinalize()  {}
 	virtual ~FLocalLightRenderState() {}
 
 	bool bStationary = false;
@@ -267,7 +269,9 @@ using FSpotLightRenderStateRef = TEntityArray<FSpotLightRenderState>::EntityRefT
 struct FRectLightRenderState : public FLocalLightRenderState
 {
 	FRectLightRenderState(URectLightComponent* ComponentUObject);
-
+	virtual void RenderThreadInit() override;
+	virtual void RenderThreadFinalize() override;
+	
 	FLinearColor Color;
 	float AttenuationRadius;
 	FVector Position;
@@ -278,6 +282,11 @@ struct FRectLightRenderState : public FLocalLightRenderState
 	float BarnDoorAngle;
 	float BarnDoorLength;
 	FTexture* IESTexture;
+	UTexture* SourceTexture;
+	uint32 AtlasSlotIndex;
+	FVector2f RectLightAtlasUVOffset;
+	FVector2f RectLightAtlasUVScale;
+	float RectLightAtlasMaxLevel;
 
 	virtual FLightRenderParameters GetLightShaderParameters() const override;
 };
diff --git a/Engine/Plugins/Experimental/GPULightmass/Source/GPULightmass/Private/Scene/Scene.cpp b/Engine/Plugins/Experimental/GPULightmass/Source/GPULightmass/Private/Scene/Scene.cpp
--- a/Engine/Plugins/Experimental/GPULightmass/Source/GPULightmass/Private/Scene/Scene.cpp
+++ b/Engine/Plugins/Experimental/GPULightmass/Source/GPULightmass/Private/Scene/Scene.cpp
@@ -442,6 +442,8 @@ void FScene::AddLight(LightComponentType* LightComponent)
 	{
 		typename LightTypeInfo<LightComponentType>::RenderStateRefType LightRenderStateRef = LightTypeInfo<LightComponentType>::GetLightRenderStateArray(RenderState.LightSceneRenderState).Emplace(MoveTemp(LightRenderState));
 
+		LightRenderStateRef->RenderThreadInit();
+
 		for (FGeometryRenderStateToken Token : RelevantGeometriesToUpdateOnRenderThread)
 		{
 			for (FLightmapRenderStateRef& Lightmap : Token.RenderStateArray->Get(Token.ElementId).LODLightmapRenderStates)
@@ -531,6 +533,8 @@ void FScene::RemoveLight(LightComponentType* PointLightComponent)
 	{
 		typename LightTypeInfo<LightComponentType>::RenderStateRefType LightRenderStateRef(LightTypeInfo<LightComponentType>::GetLightRenderStateArray(RenderState.LightSceneRenderState).Elements[ElementId], LightTypeInfo<LightComponentType>::GetLightRenderStateArray(RenderState.LightSceneRenderState));
 
+		LightRenderStateRef->RenderThreadFinalize();
+		
 		for (FGeometryRenderStateToken Token : RelevantGeometriesToUpdateOnRenderThread)
 		{
 			for (FLightmapRenderStateRef& Lightmap : Token.RenderStateArray->Get(Token.ElementId).LODLightmapRenderStates)
diff --git a/Engine/Plugins/Experimental/GPULightmass/Source/GPULightmass/Private/VolumetricLightmap.cpp b/Engine/Plugins/Experimental/GPULightmass/Source/GPULightmass/Private/VolumetricLightmap.cpp
--- a/Engine/Plugins/Experimental/GPULightmass/Source/GPULightmass/Private/VolumetricLightmap.cpp
+++ b/Engine/Plugins/Experimental/GPULightmass/Source/GPULightmass/Private/VolumetricLightmap.cpp
@@ -10,6 +10,7 @@
 #include "RHIGPUReadback.h"
 #include "LevelEditorViewport.h"
 #include "Editor.h"
+#include "RectLightTextureManager.h"
 
 IMPLEMENT_MATERIAL_SHADER_TYPE(, FVLMVoxelizationVS, TEXT("/Plugin/GPULightmass/Private/VolumetricLightmapVoxelization.usf"), TEXT("VLMVoxelizationVS"), SF_Vertex);
 IMPLEMENT_MATERIAL_SHADER_TYPE(, FVLMVoxelizationGS, TEXT("/Plugin/GPULightmass/Private/VolumetricLightmapVoxelization.usf"), TEXT("VLMVoxelizationGS"), SF_Geometry);
@@ -526,7 +527,10 @@ void FVolumetricLightmapRenderer::BackgroundTick()
 		}
 	}
 
-	FRDGBuilder GraphBuilder(FRHICommandListExecutor::GetImmediateCommandList());
+	FRDGBuilder GraphBuilder(FRHICommandListExecutor::GetImmediateCommandList());	
+
+	RectLightAtlas::UpdateRectLightAtlasTexture(GraphBuilder, Scene->FeatureLevel);
+	
 	{
 		RDG_EVENT_SCOPE(GraphBuilder, "Volumetric Lightmap Path Tracing");
 
diff --git a/Engine/Source/Runtime/Renderer/Private/RectLightTextureManager.h b/Engine/Source/Runtime/Renderer/Private/RectLightTextureManager.h
--- a/Engine/Source/Runtime/Renderer/Private/RectLightTextureManager.h
+++ b/Engine/Source/Runtime/Renderer/Private/RectLightTextureManager.h
@@ -20,22 +20,22 @@ struct FAtlasSlotDesc
 };
 	
 // Add a rect light source texture to the texture atlas
-uint32 AddRectLightTexture(UTexture* Texture);
+RENDERER_API uint32 AddRectLightTexture(UTexture* Texture);
 
 // Remove a rect light source texture to the texture atlas
-void RemoveRectLightTexture(uint32 InSlotId);
+RENDERER_API void RemoveRectLightTexture(uint32 InSlotId);
 
 // Return the atlas texture coordinate for a	 particular slot
-FAtlasSlotDesc GetRectLightAtlasSlot(uint32 InSlotId);
+RENDERER_API FAtlasSlotDesc GetRectLightAtlasSlot(uint32 InSlotId);
 
 // Return the atlas texture
-FRHITexture* GetRectLightAtlasTexture();
+RENDERER_API FRHITexture* GetRectLightAtlasTexture();
 
 // Update the rect light atlas texture
-void UpdateRectLightAtlasTexture(FRDGBuilder& GraphBuilder, const ERHIFeatureLevel::Type FeatureLevel);
+RENDERER_API void UpdateRectLightAtlasTexture(FRDGBuilder& GraphBuilder, const ERHIFeatureLevel::Type FeatureLevel);
 
 // Return the rect light atlas debug pass
-void AddRectLightAtlasDebugPass(FRDGBuilder& GraphBuilder, const FViewInfo& View, FRDGTextureRef OutputTexture);
+RENDERER_API void AddRectLightAtlasDebugPass(FRDGBuilder& GraphBuilder, const FViewInfo& View, FRDGTextureRef OutputTexture);
 
 // Scope for invalidating a particular texture 
 // This ensures the atlas contains the latest version of the texture and filter it
