GPULM: Further mitigate denoiser discoloration issue under extreme dynamic ranges by applying simple tonemapping (sqrt()) before denoising

diff --git a/Engine/Plugins/Experimental/GPULightmass/Source/GPULightmass/Private/LightmapDenoising.cpp b/Engine/Plugins/Experimental/GPULightmass/Source/GPULightmass/Private/LightmapDenoising.cpp
--- a/Engine/Plugins/Experimental/GPULightmass/Source/GPULightmass/Private/LightmapDenoising.cpp
+++ b/Engine/Plugins/Experimental/GPULightmass/Source/GPULightmass/Private/LightmapDenoising.cpp
@@ -45,6 +45,22 @@ void FDenoiserFilterSet::Clear()
 	FMemory::Memset(OutputBuffer.GetData(), 0, OutputBuffer.GetTypeSize() * OutputBuffer.Num());
 }
 
+FLinearColor TonemapLightingForDenoising(FLinearColor Lighting)
+{
+	Lighting.R = FMath::Sqrt(Lighting.R);
+	Lighting.G = FMath::Sqrt(Lighting.G);
+	Lighting.B = FMath::Sqrt(Lighting.B);
+	return Lighting;
+}
+
+FLinearColor InverseTonemapLightingForDenoising(FLinearColor Lighting)
+{
+	Lighting.R = FMath::Square(Lighting.R);
+	Lighting.G = FMath::Square(Lighting.G);
+	Lighting.B = FMath::Square(Lighting.B);
+	return Lighting;
+}
+
 void DenoiseLightSampleData(FIntPoint Size, TArray<FLightSampleData>& LightSampleData, FDenoiserContext& DenoiserContext, bool bPrepadTexels)
 {
 	if (bPrepadTexels)
@@ -105,9 +121,18 @@ void DenoiseLightSampleData(FIntPoint Size, TArray<FLightSampleData>& LightSampl
 		{
 			for (int32 X = 0; X < Size.X; X++)
 			{
-				FilterSet.InputBuffer[Y * FilterSet.Size.X + X][0] = LightSampleData[Y * Size.X + X].Coefficients[0][0];
-				FilterSet.InputBuffer[Y * FilterSet.Size.X + X][1] = LightSampleData[Y * Size.X + X].Coefficients[0][1];
-				FilterSet.InputBuffer[Y * FilterSet.Size.X + X][2] = LightSampleData[Y * Size.X + X].Coefficients[0][2];
+				if (LightSampleData[Y * Size.X + X].bIsMapped)
+				{
+					FLinearColor IncidentLighting { EForceInit::ForceInitToZero };
+					IncidentLighting.R = LightSampleData[Y * Size.X + X].Coefficients[0][0];
+					IncidentLighting.G = LightSampleData[Y * Size.X + X].Coefficients[0][1];
+					IncidentLighting.B = LightSampleData[Y * Size.X + X].Coefficients[0][2];
+					IncidentLighting = TonemapLightingForDenoising(IncidentLighting);
+				
+					FilterSet.InputBuffer[Y * FilterSet.Size.X + X][0] = IncidentLighting.R;
+					FilterSet.InputBuffer[Y * FilterSet.Size.X + X][1] = IncidentLighting.G;
+					FilterSet.InputBuffer[Y * FilterSet.Size.X + X][2] = IncidentLighting.B;
+				}
 			}
 		}
 
@@ -117,14 +142,20 @@ void DenoiseLightSampleData(FIntPoint Size, TArray<FLightSampleData>& LightSampl
 		{
 			for (int32 X = 0; X < Size.X; X++)
 			{
-				LightSampleData[Y * Size.X + X].Coefficients[0][0] = FilterSet.OutputBuffer[Y * FilterSet.Size.X + X][0];
-				LightSampleData[Y * Size.X + X].Coefficients[0][1] = FilterSet.OutputBuffer[Y * FilterSet.Size.X + X][1];
-				LightSampleData[Y * Size.X + X].Coefficients[0][2] = FilterSet.OutputBuffer[Y * FilterSet.Size.X + X][2];
+				FLinearColor DenoisedLighting { EForceInit::ForceInitToZero };
+				DenoisedLighting.R = FilterSet.OutputBuffer[Y * FilterSet.Size.X + X][0];
+				DenoisedLighting.G = FilterSet.OutputBuffer[Y * FilterSet.Size.X + X][1];
+				DenoisedLighting.B = FilterSet.OutputBuffer[Y * FilterSet.Size.X + X][2];
+				DenoisedLighting = InverseTonemapLightingForDenoising(DenoisedLighting);
+				
+				LightSampleData[Y * Size.X + X].Coefficients[0][0] = DenoisedLighting.R;
+				LightSampleData[Y * Size.X + X].Coefficients[0][1] = DenoisedLighting.G;
+				LightSampleData[Y * Size.X + X].Coefficients[0][2] = DenoisedLighting.B;
 
 				// Copy to LQ coefficients
-				LightSampleData[Y * Size.X + X].Coefficients[2][0] = FilterSet.OutputBuffer[Y * FilterSet.Size.X + X][0];
-				LightSampleData[Y * Size.X + X].Coefficients[2][1] = FilterSet.OutputBuffer[Y * FilterSet.Size.X + X][1];
-				LightSampleData[Y * Size.X + X].Coefficients[2][2] = FilterSet.OutputBuffer[Y * FilterSet.Size.X + X][2];
+				LightSampleData[Y * Size.X + X].Coefficients[2][0] = DenoisedLighting.R;
+				LightSampleData[Y * Size.X + X].Coefficients[2][1] = DenoisedLighting.G;
+				LightSampleData[Y * Size.X + X].Coefficients[2][2] = DenoisedLighting.B;
 			}
 		}
 	}
@@ -138,9 +169,12 @@ void DenoiseLightSampleData(FIntPoint Size, TArray<FLightSampleData>& LightSampl
 		{
 			for (int32 X = 0; X < Size.X; X++)
 			{
-				FilterSet.InputBuffer[Y * FilterSet.Size.X + X][0] = LightSampleData[Y * Size.X + X].Coefficients[1][0];
-				FilterSet.InputBuffer[Y * FilterSet.Size.X + X][1] = LightSampleData[Y * Size.X + X].Coefficients[1][1];
-				FilterSet.InputBuffer[Y * FilterSet.Size.X + X][2] = LightSampleData[Y * Size.X + X].Coefficients[1][2];
+				if (LightSampleData[Y * Size.X + X].bIsMapped)
+				{
+					FilterSet.InputBuffer[Y * FilterSet.Size.X + X][0] = LightSampleData[Y * Size.X + X].Coefficients[1][0];
+					FilterSet.InputBuffer[Y * FilterSet.Size.X + X][1] = LightSampleData[Y * Size.X + X].Coefficients[1][1];
+					FilterSet.InputBuffer[Y * FilterSet.Size.X + X][2] = LightSampleData[Y * Size.X + X].Coefficients[1][2];
+				}
 			}
 		}
 
@@ -263,14 +297,17 @@ void DenoiseRawData(
 		// Resizing the filter is a super expensive operation
 		// Round things into size bins to reduce number of resizes
 		FDenoiserFilterSet& FilterSet = DenoiserContext.GetFilterForSize(FIntPoint(FMath::DivideAndRoundUp(Size.X, 64) * 64, FMath::DivideAndRoundUp(Size.Y, 64) * 64));
-
+		
 		for (int32 Y = 0; Y < Size.Y; Y++)
 		{
 			for (int32 X = 0; X < Size.X; X++)
 			{
-				FilterSet.InputBuffer[Y * FilterSet.Size.X + X][0] = IncidentLighting[Y * Size.X + X].R;
-				FilterSet.InputBuffer[Y * FilterSet.Size.X + X][1] = IncidentLighting[Y * Size.X + X].G;
-				FilterSet.InputBuffer[Y * FilterSet.Size.X + X][2] = IncidentLighting[Y * Size.X + X].B;
+				if (IncidentLighting[Y * Size.X + X].A >= 0)
+				{
+					FilterSet.InputBuffer[Y * FilterSet.Size.X + X][0] = TonemapLightingForDenoising(IncidentLighting[Y * Size.X + X]).R;
+					FilterSet.InputBuffer[Y * FilterSet.Size.X + X][1] = TonemapLightingForDenoising(IncidentLighting[Y * Size.X + X]).G;
+					FilterSet.InputBuffer[Y * FilterSet.Size.X + X][2] = TonemapLightingForDenoising(IncidentLighting[Y * Size.X + X]).B;
+				}
 			}
 		}
 
@@ -280,9 +317,9 @@ void DenoiseRawData(
 		{
 			for (int32 X = 0; X < Size.X; X++)
 			{
-				IncidentLighting[Y * Size.X + X].R = FilterSet.OutputBuffer[Y * FilterSet.Size.X + X][0];
-				IncidentLighting[Y * Size.X + X].G = FilterSet.OutputBuffer[Y * FilterSet.Size.X + X][1];
-				IncidentLighting[Y * Size.X + X].B = FilterSet.OutputBuffer[Y * FilterSet.Size.X + X][2];
+				IncidentLighting[Y * Size.X + X].R = InverseTonemapLightingForDenoising(FLinearColor(FilterSet.OutputBuffer[Y * FilterSet.Size.X + X])).R;
+				IncidentLighting[Y * Size.X + X].G = InverseTonemapLightingForDenoising(FLinearColor(FilterSet.OutputBuffer[Y * FilterSet.Size.X + X])).G;
+				IncidentLighting[Y * Size.X + X].B = InverseTonemapLightingForDenoising(FLinearColor(FilterSet.OutputBuffer[Y * FilterSet.Size.X + X])).B;
 			}
 		}
 	}
@@ -292,15 +329,16 @@ void DenoiseRawData(
 		// Round things into size bins to reduce number of resizes
 		FDenoiserFilterSet& FilterSet = DenoiserContext.GetFilterForSize(FIntPoint(FMath::DivideAndRoundUp(Size.X, 64) * 64, FMath::DivideAndRoundUp(Size.Y, 64) * 64), true);
 
-		FLinearColor MinValue(MAX_flt, MAX_flt, MAX_flt);
-
 		for (int32 Y = 0; Y < Size.Y; Y++)
 		{
 			for (int32 X = 0; X < Size.X; X++)
 			{
-				FilterSet.InputBuffer[Y * FilterSet.Size.X + X][0] = LuminanceSH[Y * Size.X + X].R;
-				FilterSet.InputBuffer[Y * FilterSet.Size.X + X][1] = LuminanceSH[Y * Size.X + X].G;
-				FilterSet.InputBuffer[Y * FilterSet.Size.X + X][2] = LuminanceSH[Y * Size.X + X].B;
+				if (IncidentLighting[Y * Size.X + X].A >= 0)
+				{
+					FilterSet.InputBuffer[Y * FilterSet.Size.X + X][0] = LuminanceSH[Y * Size.X + X].R;
+					FilterSet.InputBuffer[Y * FilterSet.Size.X + X][1] = LuminanceSH[Y * Size.X + X].G;
+					FilterSet.InputBuffer[Y * FilterSet.Size.X + X][2] = LuminanceSH[Y * Size.X + X].B;
+				}
 			}
 		}
 
