GPULM: Implement SimpleFireflyRemover as an alternative to Intel OpenImageDenoise.

diff --git a/Engine/Plugins/Experimental/GPULightmass/Source/GPULightmass/Private/LightmapDenoising.cpp b/Engine/Plugins/Experimental/GPULightmass/Source/GPULightmass/Private/LightmapDenoising.cpp
--- a/Engine/Plugins/Experimental/GPULightmass/Source/GPULightmass/Private/LightmapDenoising.cpp
+++ b/Engine/Plugins/Experimental/GPULightmass/Source/GPULightmass/Private/LightmapDenoising.cpp
@@ -61,150 +61,139 @@ FLinearColor InverseTonemapLightingForDenoising(FLinearColor Lighting)
 	return Lighting;
 }
 
-void DenoiseLightSampleData(FIntPoint Size, TArray<FLightSampleData>& LightSampleData, FDenoiserContext& DenoiserContext, bool bPrepadTexels)
+template<typename LightSampleType>
+struct LightSampleDataProvider {};
+
+template<>
+struct LightSampleDataProvider<FLinearColor>
 {
-	if (bPrepadTexels)
+	FIntPoint Size;
+	TArray<FLinearColor>& IncidentLighting;
+	TArray<FLinearColor>& LuminanceSH;
+
+	LightSampleDataProvider<FLinearColor>(
+		FIntPoint Size,
+		TArray<FLinearColor>& IncidentLighting,
+		TArray<FLinearColor>& LuminanceSH)
+		: Size(Size)
+		, IncidentLighting(IncidentLighting)
+		, LuminanceSH(LuminanceSH)
+	{
+	}
+
+	bool IsMapped(FIntPoint Location)
 	{
-		TArray<int32> DilationMask;
-		DilationMask.AddZeroed(Size.X * Size.Y);
+		return IncidentLighting[Location.Y * Size.X + Location.X].A >= 0;
+	}
+	
+	float GetL(FIntPoint Location)
+	{
+		return IncidentLighting[Location.Y * Size.X + Location.X].GetLuminance();
+	}
+
+	void OverwriteTexel(FIntPoint Dst, FIntPoint Src)
+	{
+		IncidentLighting[Dst.Y * Size.X + Dst.X] = IncidentLighting[Src.Y * Size.X + Src.X];
+		LuminanceSH[Dst.Y * Size.X + Dst.X] = LuminanceSH[Src.Y * Size.X + Src.X];
+	}
+};
+
+const int32 WindowSize = 2;
+
+int32 EncodeOffset(FIntPoint D)
+{
+	return (D.Y + WindowSize) * (WindowSize * 2 + 1) + D.X + WindowSize;
+}
+
+FIntPoint DecodeOffset(int32 EncodedIndex)
+{
+	FIntPoint D;
+	D.Y = EncodedIndex / (WindowSize * 2 + 1) - WindowSize;
+	D.X = EncodedIndex % (WindowSize * 2 + 1) - WindowSize;
+	return D;
+}
 
-		const int32 TotalIteration = 2;
+bool IsColorValid(FLinearColor Color)
+{
+	return FMath::IsFinite(Color.R) && FMath::IsFinite(Color.G) && FMath::IsFinite(Color.B) && FMath::IsFinite(Color.A)
+	&& Color.R >= 0 && Color.G >= 0 && Color.B >= 0 && Color.A >= 0;
+}
 
-		for (int32 IterationIndex = 0; IterationIndex < TotalIteration; IterationIndex++)
+void SimpleFireflyFilter(
+	FIntPoint Size,
+	TArray<FLinearColor>& IncidentLighting,
+	TArray<FLinearColor>& LuminanceSH)
+{
+	LightSampleDataProvider<FLinearColor> SampleData(Size, IncidentLighting, LuminanceSH);
+	
+	TArray<int32> Buffer;
+	Buffer.AddDefaulted(Size.X * Size.Y);
+	
+	ParallelFor(Size.Y, [&](int32 Y)
+	{
+		for (int32 X = 0; X < Size.X; X++)
 		{
-			for (int32 Y = 0; Y < Size.Y; Y++)
+			if (SampleData.IsMapped({X, Y}))
 			{
-				for (int32 X = 0; X < Size.X; X++)
+				int32 ValidCount = 0;
+				int32 BiggerCount = 0;
+				TArray<int32, TFixedAllocator<(WindowSize * 2 + 1) * (WindowSize * 2 + 1)>> ValidTexelIndices;
+				
+				for (int32 dx = -WindowSize; dx <= WindowSize; dx++)
 				{
-					if (!LightSampleData[Y * Size.X + X].bIsMapped)
+					for (int32 dy = -WindowSize; dy <= WindowSize; dy++)
 					{
-						for (int32 dx = -1; dx <= 1; dx++)
+						if (Y + dy >= 0 && Y + dy < Size.Y && X + dx >= 0 && X + dx < Size.X)
 						{
-							for (int32 dy = -1; dy <= 1; dy++)
+							if (SampleData.IsMapped(FIntPoint {X, Y} + FIntPoint {dx, dy}))
 							{
-								if (Y + dy >= 0 && Y + dy < Size.Y && X + dx >= 0 && X + dx < Size.X)
+								ValidCount++;
+								ValidTexelIndices.Add(EncodeOffset({dx, dy}));
+
+								if (SampleData.GetL({X, Y}) >= 1.01 * SampleData.GetL(FIntPoint {X, Y} + FIntPoint {dx, dy}))
 								{
-									if (LightSampleData[(Y + dy) * Size.X + (X + dx)].bIsMapped)
-									{
-										LightSampleData[Y * Size.X + X] = LightSampleData[(Y + dy) * Size.X + (X + dx)];
-										LightSampleData[Y * Size.X + X].bIsMapped = false;
-										DilationMask[Y * Size.X + X] = 1;
-									}
+									BiggerCount++;
 								}
 							}
 						}
 					}
 				}
-			}
 
-			for (int32 Y = 0; Y < Size.Y; Y++)
-			{
-				for (int32 X = 0; X < Size.X; X++)
+				if (BiggerCount >= 1 && BiggerCount > ValidCount * 0.75)
 				{
-					if (DilationMask[Y * Size.X + X] == 1)
+					ValidTexelIndices.Sort([&SampleData, X, Y](const int32& A, const int32& B)
 					{
-						DilationMask[Y * Size.X + X] = 0;
-						LightSampleData[Y * Size.X + X].bIsMapped = true;
-					}
-				}
-			}
-		}
-	}
+						return SampleData.GetL(FIntPoint {X, Y} + DecodeOffset(A)) > SampleData.GetL(FIntPoint {X, Y} + DecodeOffset(B));
+					});
 
-	{
-
-		// Resizing the filter is a super expensive operation
-		// Round things into size bins to reduce number of resizes
-		FDenoiserFilterSet& FilterSet = DenoiserContext.GetFilterForSize(FIntPoint(FMath::DivideAndRoundUp(Size.X, 64) * 64, FMath::DivideAndRoundUp(Size.Y, 64) * 64));
-
-		for (int32 Y = 0; Y < Size.Y; Y++)
-		{
-			for (int32 X = 0; X < Size.X; X++)
-			{
-				if (LightSampleData[Y * Size.X + X].bIsMapped)
-				{
-					FLinearColor IncidentLighting { EForceInit::ForceInitToZero };
-					IncidentLighting.R = LightSampleData[Y * Size.X + X].Coefficients[0][0];
-					IncidentLighting.G = LightSampleData[Y * Size.X + X].Coefficients[0][1];
-					IncidentLighting.B = LightSampleData[Y * Size.X + X].Coefficients[0][2];
-					IncidentLighting = TonemapLightingForDenoising(IncidentLighting);
-				
-					FilterSet.InputBuffer[Y * FilterSet.Size.X + X][0] = IncidentLighting.R;
-					FilterSet.InputBuffer[Y * FilterSet.Size.X + X][1] = IncidentLighting.G;
-					FilterSet.InputBuffer[Y * FilterSet.Size.X + X][2] = IncidentLighting.B;
+					Buffer[Y * Size.X + X] = ValidTexelIndices[ValidCount / 2];
 				}
-			}
-		}
-
-		FilterSet.Execute();
-
-		for (int32 Y = 0; Y < Size.Y; Y++)
-		{
-			for (int32 X = 0; X < Size.X; X++)
-			{
-				FLinearColor DenoisedLighting { EForceInit::ForceInitToZero };
-				DenoisedLighting.R = FilterSet.OutputBuffer[Y * FilterSet.Size.X + X][0];
-				DenoisedLighting.G = FilterSet.OutputBuffer[Y * FilterSet.Size.X + X][1];
-				DenoisedLighting.B = FilterSet.OutputBuffer[Y * FilterSet.Size.X + X][2];
-				DenoisedLighting = InverseTonemapLightingForDenoising(DenoisedLighting);
-				
-				LightSampleData[Y * Size.X + X].Coefficients[0][0] = DenoisedLighting.R;
-				LightSampleData[Y * Size.X + X].Coefficients[0][1] = DenoisedLighting.G;
-				LightSampleData[Y * Size.X + X].Coefficients[0][2] = DenoisedLighting.B;
-
-				// Copy to LQ coefficients
-				LightSampleData[Y * Size.X + X].Coefficients[2][0] = DenoisedLighting.R;
-				LightSampleData[Y * Size.X + X].Coefficients[2][1] = DenoisedLighting.G;
-				LightSampleData[Y * Size.X + X].Coefficients[2][2] = DenoisedLighting.B;
-			}
-		}
-	}
-
-	{
-		// Resizing the filter is a super expensive operation
-		// Round things into size bins to reduce number of resizes
-		FDenoiserFilterSet& FilterSet = DenoiserContext.GetFilterForSize(FIntPoint(FMath::DivideAndRoundUp(Size.X, 64) * 64, FMath::DivideAndRoundUp(Size.Y, 64) * 64), true);
-
-		for (int32 Y = 0; Y < Size.Y; Y++)
-		{
-			for (int32 X = 0; X < Size.X; X++)
-			{
-				if (LightSampleData[Y * Size.X + X].bIsMapped)
+				else
 				{
-					FilterSet.InputBuffer[Y * FilterSet.Size.X + X][0] = LightSampleData[Y * Size.X + X].Coefficients[1][0];
-					FilterSet.InputBuffer[Y * FilterSet.Size.X + X][1] = LightSampleData[Y * Size.X + X].Coefficients[1][1];
-					FilterSet.InputBuffer[Y * FilterSet.Size.X + X][2] = LightSampleData[Y * Size.X + X].Coefficients[1][2];
+					Buffer[Y * Size.X + X] = EncodeOffset({0, 0});
 				}
 			}
 		}
-
-		FilterSet.Execute();
-
-		for (int32 Y = 0; Y < Size.Y; Y++)
+	});
+		
+	ParallelFor(Size.Y, [&](int32 Y)
+	{
+		for (int32 X = 0; X < Size.X; X++)
 		{
-			for (int32 X = 0; X < Size.X; X++)
+			if (SampleData.IsMapped({X, Y}))
 			{
-				LightSampleData[Y * Size.X + X].Coefficients[1][0] = FilterSet.OutputBuffer[Y * FilterSet.Size.X + X][0];
-				LightSampleData[Y * Size.X + X].Coefficients[1][1] = FilterSet.OutputBuffer[Y * FilterSet.Size.X + X][1];
-				LightSampleData[Y * Size.X + X].Coefficients[1][2] = FilterSet.OutputBuffer[Y * FilterSet.Size.X + X][2];
-
-				// Copy to LQ coefficients
-				LightSampleData[Y * Size.X + X].Coefficients[3][0] = FilterSet.OutputBuffer[Y * FilterSet.Size.X + X][0];
-				LightSampleData[Y * Size.X + X].Coefficients[3][1] = FilterSet.OutputBuffer[Y * FilterSet.Size.X + X][1];
-				LightSampleData[Y * Size.X + X].Coefficients[3][2] = FilterSet.OutputBuffer[Y * FilterSet.Size.X + X][2];
+				SampleData.OverwriteTexel({X, Y}, FIntPoint {X, Y} + DecodeOffset(Buffer[Y * Size.X + X]));
 			}
 		}
-	}
+	});
 }
 
-void DenoiseRawData(
+void NormalizeDirectionalLuminancePair(
 	FIntPoint Size,
 	TArray<FLinearColor>& IncidentLighting,
-	TArray<FLinearColor>& LuminanceSH,
-	FDenoiserContext& DenoiserContext, 
-	bool bPrepadTexels)
+	TArray<FLinearColor>& LuminanceSH)
 {
-	for (int32 Y = 0; Y < Size.Y; Y++)
+	ParallelFor(Size.Y, [&](int32 Y)
 	{
 		for (int32 X = 0; X < Size.X; X++)
 		{
@@ -242,57 +231,86 @@ void DenoiseRawData(
 				LuminanceSH[Y * Size.X + X].B = 0;
 			}
 		}
-	}
+	});
+}
 
-	if (bPrepadTexels)
-	{
-		TArray<float> DilationMask;
-		DilationMask.AddZeroed(Size.X * Size.Y);
+void DilateDirectionalLuminancePairForDenoising(
+	FIntPoint Size,
+	TArray<FLinearColor>& IncidentLighting,
+	TArray<FLinearColor>& LuminanceSH)
+{
+	TArray<float> DilationMask;
+	DilationMask.AddZeroed(Size.X * Size.Y);
 
-		const int32 TotalIteration = 2;
+	const int32 TotalIteration = 2;
 
-		for (int32 IterationIndex = 0; IterationIndex < TotalIteration; IterationIndex++)
+	for (int32 IterationIndex = 0; IterationIndex < TotalIteration; IterationIndex++)
+	{
+		for (int32 Y = 0; Y < Size.Y; Y++)
 		{
-			for (int32 Y = 0; Y < Size.Y; Y++)
+			for (int32 X = 0; X < Size.X; X++)
 			{
-				for (int32 X = 0; X < Size.X; X++)
+				if (!(IncidentLighting[Y * Size.X + X].A >= 0.0f))
 				{
-					if (!(IncidentLighting[Y * Size.X + X].A >= 0.0f))
+					for (int32 dx = -1; dx <= 1; dx++)
 					{
-						for (int32 dx = -1; dx <= 1; dx++)
+						for (int32 dy = -1; dy <= 1; dy++)
 						{
-							for (int32 dy = -1; dy <= 1; dy++)
+							if (Y + dy >= 0 && Y + dy < Size.Y && X + dx >= 0 && X + dx < Size.X)
 							{
-								if (Y + dy >= 0 && Y + dy < Size.Y && X + dx >= 0 && X + dx < Size.X)
+								if (IncidentLighting[(Y + dy) * Size.X + (X + dx)].A >= 0.0f)
 								{
-									if (IncidentLighting[(Y + dy) * Size.X + (X + dx)].A >= 0.0f)
-									{
-										IncidentLighting[Y * Size.X + X] = IncidentLighting[(Y + dy) * Size.X + (X + dx)];
-										LuminanceSH[Y * Size.X + X] = LuminanceSH[(Y + dy) * Size.X + (X + dx)];
-										DilationMask[Y * Size.X + X] = IncidentLighting[Y * Size.X + X].A;
-										IncidentLighting[Y * Size.X + X].A = -1.0f;
-									}
+									IncidentLighting[Y * Size.X + X] = IncidentLighting[(Y + dy) * Size.X + (X + dx)];
+									LuminanceSH[Y * Size.X + X] = LuminanceSH[(Y + dy) * Size.X + (X + dx)];
+									DilationMask[Y * Size.X + X] = IncidentLighting[Y * Size.X + X].A;
+									IncidentLighting[Y * Size.X + X].A = -1.0f;
 								}
 							}
 						}
 					}
 				}
 			}
+		}
 
-			for (int32 Y = 0; Y < Size.Y; Y++)
+		for (int32 Y = 0; Y < Size.Y; Y++)
+		{
+			for (int32 X = 0; X < Size.X; X++)
 			{
-				for (int32 X = 0; X < Size.X; X++)
+				if (DilationMask[Y * Size.X + X] > 0.0f)
 				{
-					if (DilationMask[Y * Size.X + X] > 0.0f)
-					{
-						IncidentLighting[Y * Size.X + X].A = DilationMask[Y * Size.X + X];
-						DilationMask[Y * Size.X + X] = 0;
-					}
+					IncidentLighting[Y * Size.X + X].A = DilationMask[Y * Size.X + X];
+					DilationMask[Y * Size.X + X] = 0;
 				}
 			}
 		}
 	}
+}
 
+void ReencodeDirectionalLuminancePair(
+	FIntPoint Size,
+	TArray<FLinearColor>& IncidentLighting,
+	TArray<FLinearColor>& LuminanceSH)
+{
+	ParallelFor(Size.Y, [&](int32 Y)
+	{
+		for (int32 X = 0; X < Size.X; X++)
+		{
+			IncidentLighting[Y * Size.X + X].R = FMath::Sqrt(IncidentLighting[Y * Size.X + X].R);
+			IncidentLighting[Y * Size.X + X].G = FMath::Sqrt(IncidentLighting[Y * Size.X + X].G);
+			IncidentLighting[Y * Size.X + X].B = FMath::Sqrt(IncidentLighting[Y * Size.X + X].B);
+
+			LuminanceSH[Y * Size.X + X].A *= 0.282095f;
+		}
+	});
+}
+
+void DenoiseDirectionalLuminancePair(
+	FIntPoint Size,
+	TArray<FLinearColor>& IncidentLighting,
+	TArray<FLinearColor>& LuminanceSH,
+	FDenoiserContext& DenoiserContext)
+{
+	// Denoise IncidentLighting
 	{
 		// Resizing the filter is a super expensive operation
 		// Round things into size bins to reduce number of resizes
@@ -324,6 +342,7 @@ void DenoiseRawData(
 		}
 	}
 
+	// Denoise LuminanceSH
 	{
 		// Resizing the filter is a super expensive operation
 		// Round things into size bins to reduce number of resizes
@@ -354,16 +373,21 @@ void DenoiseRawData(
 			}
 		}
 	}
+}
 
-	for (int32 Y = 0; Y < Size.Y; Y++)
+void DenoiseRawData(
+	FIntPoint Size,
+	TArray<FLinearColor>& IncidentLighting,
+	TArray<FLinearColor>& LuminanceSH,
+	FDenoiserContext& DenoiserContext, 
+	bool bPrepadTexels)
+{
+	// Required for Intel OIDN as it assumes LuminanceSH is normalized into [-1, 1]
+	NormalizeDirectionalLuminancePair(Size, IncidentLighting, LuminanceSH);
+	if (bPrepadTexels)
 	{
-		for (int32 X = 0; X < Size.X; X++)
-		{
-			IncidentLighting[Y * Size.X + X].R = FMath::Sqrt(IncidentLighting[Y * Size.X + X].R);
-			IncidentLighting[Y * Size.X + X].G = FMath::Sqrt(IncidentLighting[Y * Size.X + X].G);
-			IncidentLighting[Y * Size.X + X].B = FMath::Sqrt(IncidentLighting[Y * Size.X + X].B);
-
-			LuminanceSH[Y * Size.X + X].A *= 0.282095f;
-		}
+		DilateDirectionalLuminancePairForDenoising(Size, IncidentLighting, LuminanceSH);
 	}
+	DenoiseDirectionalLuminancePair(Size, IncidentLighting, LuminanceSH, DenoiserContext);
+	ReencodeDirectionalLuminancePair(Size, IncidentLighting, LuminanceSH);
 }
diff --git a/Engine/Plugins/Experimental/GPULightmass/Source/GPULightmass/Private/LightmapDenoising.h b/Engine/Plugins/Experimental/GPULightmass/Source/GPULightmass/Private/LightmapDenoising.h
--- a/Engine/Plugins/Experimental/GPULightmass/Source/GPULightmass/Private/LightmapDenoising.h
+++ b/Engine/Plugins/Experimental/GPULightmass/Source/GPULightmass/Private/LightmapDenoising.h
@@ -88,11 +88,14 @@ struct FDenoiserContext
 	}
 };
 
-void DenoiseLightSampleData(FIntPoint Size, TArray<FLightSampleData>& LightSampleData, FDenoiserContext& DenoiserContext, bool bPrepadTexels = true);
-
 void DenoiseRawData(
 	FIntPoint Size,
 	TArray<FLinearColor>& IncidentLighting,
 	TArray<FLinearColor>& LuminanceSH,
 	FDenoiserContext& DenoiserContext,
-	bool bPrepadTexels = true);
\ No newline at end of file
+	bool bPrepadTexels = true);
+
+void SimpleFireflyFilter(
+	FIntPoint Size,
+	TArray<FLinearColor>& IncidentLighting,
+	TArray<FLinearColor>& LuminanceSH);
diff --git a/Engine/Plugins/Experimental/GPULightmass/Source/GPULightmass/Private/LightmapRenderer.cpp b/Engine/Plugins/Experimental/GPULightmass/Source/GPULightmass/Private/LightmapRenderer.cpp
--- a/Engine/Plugins/Experimental/GPULightmass/Source/GPULightmass/Private/LightmapRenderer.cpp
+++ b/Engine/Plugins/Experimental/GPULightmass/Source/GPULightmass/Private/LightmapRenderer.cpp
@@ -3032,13 +3032,20 @@ const int32 DenoiseTileProximity = 3;
 
 void FLightmapTileDenoiseAsyncTask::DoThreadedWork()
 {
-	static FDenoiserContext DenoiserContext;
-
-	DenoiseRawData(
-		Size,
-		TextureData->Texture[0],
-		TextureData->Texture[1],
-		DenoiserContext);
+	if (Denoiser == EGPULightmassDenoiser::SimpleFireflyRemover)
+	{
+		SimpleFireflyFilter(Size, TextureData->Texture[0], TextureData->Texture[1]);
+	}
+	else
+	{
+		static FDenoiserContext DenoiserContext;
+		
+		DenoiseRawData(
+			Size,
+			TextureData->Texture[0],
+			TextureData->Texture[1],
+			DenoiserContext);
+	}
 
 	FPlatformAtomics::AtomicStore(&TextureData->bDenoisingFinished, 1);
 }
@@ -3392,6 +3399,7 @@ void FLightmapRenderer::BackgroundTick()
 			DenoiseGroup.AsyncDenoisingWork = new FLightmapTileDenoiseAsyncTask();
 			DenoiseGroup.AsyncDenoisingWork->Size = FIntPoint(DenoiseTileProximity * GPreviewLightmapVirtualTileSize, DenoiseTileProximity * GPreviewLightmapVirtualTileSize);
 			DenoiseGroup.AsyncDenoisingWork->TextureData = DenoiseGroup.TextureData;
+			DenoiseGroup.AsyncDenoisingWork->Denoiser = Scene->Settings->Denoiser;
 			DenoisingThreadPool->AddQueuedWork(DenoiseGroup.AsyncDenoisingWork);
 
 			OngoingDenoiseGroups.Add(MoveTemp(DenoiseGroup));
diff --git a/Engine/Plugins/Experimental/GPULightmass/Source/GPULightmass/Private/LightmapRenderer.h b/Engine/Plugins/Experimental/GPULightmass/Source/GPULightmass/Private/LightmapRenderer.h
--- a/Engine/Plugins/Experimental/GPULightmass/Source/GPULightmass/Private/LightmapRenderer.h
+++ b/Engine/Plugins/Experimental/GPULightmass/Source/GPULightmass/Private/LightmapRenderer.h
@@ -4,6 +4,7 @@
 
 #include "LightmapTilePool.h"
 #include "RHIGPUReadback.h"
+#include "GPULightmassSettings.h"
 // TBB suffers from extreme fragmentation problem in editor
 #include "Core/Private/HAL/Allocators/AnsiAllocator.h"
 
@@ -99,6 +100,7 @@ class FLightmapTileDenoiseAsyncTask : public IQueuedWork
 public:
 	FIntPoint Size;
 	TSharedPtr<FLightmapTileDenoiseGroup::FTextureData, ESPMode::ThreadSafe> TextureData;
+	EGPULightmassDenoiser Denoiser;
 
 	virtual void DoThreadedWork();
 	virtual void Abandon() {}
diff --git a/Engine/Plugins/Experimental/GPULightmass/Source/GPULightmass/Private/Scene/Scene.cpp b/Engine/Plugins/Experimental/GPULightmass/Source/GPULightmass/Private/Scene/Scene.cpp
--- a/Engine/Plugins/Experimental/GPULightmass/Source/GPULightmass/Private/Scene/Scene.cpp
+++ b/Engine/Plugins/Experimental/GPULightmass/Source/GPULightmass/Private/Scene/Scene.cpp
@@ -2048,6 +2048,11 @@ void FScene::ApplyFinishedLightmapsToWorld()
 						TArray<FLightSampleData> LightSampleData;
 						LightSampleData.AddZeroed(Lightmap.GetSize().X * Lightmap.GetSize().Y); // LightSampleData will have different row pitch as VT is padded to tiles
 
+						TArray<FLinearColor> IncidentLighting;
+						TArray<FLinearColor> LuminanceSH;
+						IncidentLighting.AddZeroed(Lightmap.GetSize().X * Lightmap.GetSize().Y);
+						LuminanceSH.AddZeroed(Lightmap.GetSize().X * Lightmap.GetSize().Y);
+						
 						{
 							int32 SrcRowPitchInPixels = GPreviewLightmapVirtualTileSize;
 							int32 DstRowPitchInPixels = Lightmap.GetSize().X;
@@ -2057,11 +2062,33 @@ void FScene::ApplyFinishedLightmapsToWorld()
 								FIntRect(FIntPoint(0, 0), Lightmap.GetSize()),
 								SrcRowPitchInPixels,
 								DstRowPitchInPixels,
-								[&Lightmap, &LightSampleData](int32 DstLinearIndex, FIntPoint SrcTilePosition, int32 SrcLinearIndex) mutable
+								[&Lightmap, &IncidentLighting, &LuminanceSH](int32 DstLinearIndex, FIntPoint SrcTilePosition, int32 SrcLinearIndex) mutable
 							{
-								LightSampleData[DstLinearIndex] = ConvertToLightSample(Lightmap.TileStorage[FTileVirtualCoordinates(SrcTilePosition, 0)].CPUTextureData[0]->Data[SrcLinearIndex], Lightmap.TileStorage[FTileVirtualCoordinates(SrcTilePosition, 0)].CPUTextureData[1]->Data[SrcLinearIndex]);
+								IncidentLighting[DstLinearIndex] = Lightmap.TileStorage[FTileVirtualCoordinates(SrcTilePosition, 0)].CPUTextureData[0]->Data[SrcLinearIndex];
+								LuminanceSH[DstLinearIndex] = Lightmap.TileStorage[FTileVirtualCoordinates(SrcTilePosition, 0)].CPUTextureData[1]->Data[SrcLinearIndex];
 							});
 
+							if (Settings->DenoisingOptions == EGPULightmassDenoisingOptions::OnCompletion)
+							{
+								if (Settings->Denoiser == EGPULightmassDenoiser::SimpleFireflyRemover)
+								{
+									SimpleFireflyFilter(Lightmap.GetSize(), IncidentLighting, LuminanceSH);
+								}
+								else
+								{
+									DenoiseRawData(Lightmap.GetSize(), IncidentLighting, LuminanceSH, DenoiserContext);
+								}
+							}
+							
+							for (int32 Y = 0 ; Y < Lightmap.GetSize().Y; Y++)
+							{
+								for (int32 X = 0 ; X < Lightmap.GetSize().X; X++)
+								{
+									int32 LinearIndex = Y * Lightmap.GetSize().X + X;
+									LightSampleData[LinearIndex] = ConvertToLightSample(IncidentLighting[LinearIndex], LuminanceSH[LinearIndex]);
+								}
+							}
+
 							if (bHasSkyShadowing)
 							{
 								CopyRectTiled(
@@ -2124,11 +2151,6 @@ void FScene::ApplyFinishedLightmapsToWorld()
 						}
 #endif
 
-						if (Settings->DenoisingOptions == EGPULightmassDenoisingOptions::OnCompletion)
-						{
-							DenoiseLightSampleData(Lightmap.GetSize(), LightSampleData, DenoiserContext);
-						}
-
 						FQuantizedLightmapData* QuantizedLightmapData = new FQuantizedLightmapData();
 						QuantizedLightmapData->SizeX = Lightmap.GetSize().X;
 						QuantizedLightmapData->SizeY = Lightmap.GetSize().Y;
@@ -2330,6 +2352,11 @@ void FScene::ApplyFinishedLightmapsToWorld()
 
 						for (int32 InstanceIndex = 0; InstanceIndex < InstanceGroupLightSampleData.Num(); InstanceIndex++)
 						{
+							TArray<FLinearColor> IncidentLighting;
+							TArray<FLinearColor> LuminanceSH;
+							IncidentLighting.AddZeroed(BaseLightMapWidth * BaseLightMapHeight);
+							LuminanceSH.AddZeroed(BaseLightMapWidth * BaseLightMapHeight);
+							
 							TArray<FLightSampleData>& LightSampleData = InstanceGroupLightSampleData[InstanceIndex];
 							LightSampleData.AddZeroed(BaseLightMapWidth * BaseLightMapHeight);
 							InstancedSourceQuantizedData[InstanceIndex] = MakeUnique<FQuantizedLightmapData>();
@@ -2343,16 +2370,38 @@ void FScene::ApplyFinishedLightmapsToWorld()
 							{
 								FIntPoint InstanceTilePos = FIntPoint(RenderIndex % InstancesPerRow, RenderIndex / InstancesPerRow);
 								FIntPoint InstanceTileMin = FIntPoint(InstanceTilePos.X * BaseLightMapWidth, InstanceTilePos.Y * BaseLightMapHeight);
-
+								
 								CopyRectTiled(
 									InstanceTileMin,
 									FIntRect(FIntPoint(0, 0), FIntPoint(BaseLightMapWidth, BaseLightMapHeight)),
 									SrcRowPitchInPixels,
 									DstRowPitchInPixels,
-									[&Lightmap, &LightSampleData](int32 DstLinearIndex, FIntPoint SrcTilePosition, int32 SrcLinearIndex) mutable
+									[&Lightmap, &IncidentLighting, &LuminanceSH](int32 DstLinearIndex, FIntPoint SrcTilePosition, int32 SrcLinearIndex) mutable
 								{
-									LightSampleData[DstLinearIndex] = ConvertToLightSample(Lightmap.TileStorage[FTileVirtualCoordinates(SrcTilePosition, 0)].CPUTextureData[0]->Data[SrcLinearIndex], Lightmap.TileStorage[FTileVirtualCoordinates(SrcTilePosition, 0)].CPUTextureData[1]->Data[SrcLinearIndex]);
+									IncidentLighting[DstLinearIndex] = Lightmap.TileStorage[FTileVirtualCoordinates(SrcTilePosition, 0)].CPUTextureData[0]->Data[SrcLinearIndex];
+									LuminanceSH[DstLinearIndex] = Lightmap.TileStorage[FTileVirtualCoordinates(SrcTilePosition, 0)].CPUTextureData[1]->Data[SrcLinearIndex];
 								});
+								
+								if (Settings->DenoisingOptions == EGPULightmassDenoisingOptions::OnCompletion)
+								{
+									if (Settings->Denoiser == EGPULightmassDenoiser::SimpleFireflyRemover)
+									{
+										SimpleFireflyFilter(FIntPoint{BaseLightMapWidth, BaseLightMapHeight}, IncidentLighting, LuminanceSH);
+									}
+									else
+									{
+										DenoiseRawData(FIntPoint{BaseLightMapWidth, BaseLightMapHeight}, IncidentLighting, LuminanceSH, DenoiserContext);
+									}
+								}
+
+								for (int32 Y = 0; Y < BaseLightMapHeight; Y++)
+								{
+									for (int32 X = 0; X < BaseLightMapWidth; X++)
+									{
+										int32 LinearIndex = Y * BaseLightMapWidth + X;
+										LightSampleData[LinearIndex] = ConvertToLightSample(IncidentLighting[LinearIndex], LuminanceSH[LinearIndex]);
+									}
+								}
 
 								if (bHasSkyShadowing)
 								{
@@ -2368,11 +2417,6 @@ void FScene::ApplyFinishedLightmapsToWorld()
 								}
 							}
 
-							if (Settings->DenoisingOptions == EGPULightmassDenoisingOptions::OnCompletion)
-							{
-								DenoiseLightSampleData(FIntPoint(BaseLightMapWidth, BaseLightMapHeight), LightSampleData, DenoiserContext);
-							}
-
 							FQuantizedLightmapData& QuantizedLightmapData = *InstancedSourceQuantizedData[InstanceIndex];
 							QuantizedLightmapData.SizeX = BaseLightMapWidth;
 							QuantizedLightmapData.SizeY = BaseLightMapHeight;
@@ -2603,6 +2647,11 @@ void FScene::ApplyFinishedLightmapsToWorld()
 						TArray<FLightSampleData> LightSampleData;
 						LightSampleData.AddZeroed(Lightmap.GetSize().X * Lightmap.GetSize().Y); // LightSampleData will have different row pitch as VT is padded to tiles
 
+						TArray<FLinearColor> IncidentLighting;
+						TArray<FLinearColor> LuminanceSH;
+						IncidentLighting.AddZeroed(Lightmap.GetSize().X * Lightmap.GetSize().Y);
+						LuminanceSH.AddZeroed(Lightmap.GetSize().X * Lightmap.GetSize().Y);
+						
 						{
 							int32 SrcRowPitchInPixels = GPreviewLightmapVirtualTileSize;
 							int32 DstRowPitchInPixels = Lightmap.GetSize().X;
@@ -2612,11 +2661,33 @@ void FScene::ApplyFinishedLightmapsToWorld()
 								FIntRect(FIntPoint(0, 0), Lightmap.GetSize()),
 								SrcRowPitchInPixels,
 								DstRowPitchInPixels,
-								[&Lightmap, &LightSampleData](int32 DstLinearIndex, FIntPoint SrcTilePosition, int32 SrcLinearIndex) mutable
+								[&Lightmap, &IncidentLighting, &LuminanceSH](int32 DstLinearIndex, FIntPoint SrcTilePosition, int32 SrcLinearIndex) mutable
 							{
-								LightSampleData[DstLinearIndex] = ConvertToLightSample(Lightmap.TileStorage[FTileVirtualCoordinates(SrcTilePosition, 0)].CPUTextureData[0]->Data[SrcLinearIndex], Lightmap.TileStorage[FTileVirtualCoordinates(SrcTilePosition, 0)].CPUTextureData[1]->Data[SrcLinearIndex]);
+								IncidentLighting[DstLinearIndex] = Lightmap.TileStorage[FTileVirtualCoordinates(SrcTilePosition, 0)].CPUTextureData[0]->Data[SrcLinearIndex];
+								LuminanceSH[DstLinearIndex] = Lightmap.TileStorage[FTileVirtualCoordinates(SrcTilePosition, 0)].CPUTextureData[1]->Data[SrcLinearIndex];
 							});
 
+							if (Settings->DenoisingOptions == EGPULightmassDenoisingOptions::OnCompletion)
+							{
+								if (Settings->Denoiser == EGPULightmassDenoiser::SimpleFireflyRemover)
+								{
+									SimpleFireflyFilter(Lightmap.GetSize(), IncidentLighting, LuminanceSH);
+								}
+								else
+								{
+									DenoiseRawData(Lightmap.GetSize(), IncidentLighting, LuminanceSH, DenoiserContext);
+								}
+							}
+							
+							for (int32 Y = 0 ; Y < Lightmap.GetSize().Y; Y++)
+							{
+								for (int32 X = 0 ; X < Lightmap.GetSize().X; X++)
+								{
+									int32 LinearIndex = Y * Lightmap.GetSize().X + X;
+									LightSampleData[LinearIndex] = ConvertToLightSample(IncidentLighting[LinearIndex], LuminanceSH[LinearIndex]);
+								}
+							}
+
 							if (bHasSkyShadowing)
 							{
 								CopyRectTiled(
@@ -2631,11 +2702,6 @@ void FScene::ApplyFinishedLightmapsToWorld()
 							}
 						}
 
-						if (Settings->DenoisingOptions == EGPULightmassDenoisingOptions::OnCompletion)
-						{
-							DenoiseLightSampleData(Lightmap.GetSize(), LightSampleData, DenoiserContext);
-						}
-
 						FQuantizedLightmapData* QuantizedLightmapData = new FQuantizedLightmapData();
 						QuantizedLightmapData->SizeX = Lightmap.GetSize().X;
 						QuantizedLightmapData->SizeY = Lightmap.GetSize().Y;
diff --git a/Engine/Plugins/Experimental/GPULightmass/Source/GPULightmass/Public/GPULightmassSettings.h b/Engine/Plugins/Experimental/GPULightmass/Source/GPULightmass/Public/GPULightmassSettings.h
--- a/Engine/Plugins/Experimental/GPULightmass/Source/GPULightmass/Public/GPULightmassSettings.h
+++ b/Engine/Plugins/Experimental/GPULightmass/Source/GPULightmass/Public/GPULightmassSettings.h
@@ -22,6 +22,13 @@ enum class EGPULightmassDenoisingOptions : uint8
 	DuringInteractivePreview
 };
 
+UENUM()
+enum class EGPULightmassDenoiser : uint8
+{
+	IntelOIDN  UMETA(DisplayName = "Intel Open Image Denoise"),
+	SimpleFireflyRemover
+};
+
 UCLASS(BlueprintType)
 class GPULIGHTMASS_API UGPULightmassSettings : public UObject
 {
@@ -43,8 +50,11 @@ public:
 
 	// If enabled, denoise the results on the CPU after rendering. On Completion denoises the entire lightmap when it is finished.
 	// During Interactive Preview denoises each tile as it finishes, which is useful for previewing but less efficient.
-	UPROPERTY(EditAnywhere, BlueprintReadWrite, Category = General, DisplayName = "Denoise")
+	UPROPERTY(EditAnywhere, BlueprintReadWrite, Category = Denoising, DisplayName = "Denoise when")
 	EGPULightmassDenoisingOptions DenoisingOptions = EGPULightmassDenoisingOptions::OnCompletion;
+	
+	UPROPERTY(EditAnywhere, BlueprintReadWrite, Category = Denoising)
+	EGPULightmassDenoiser Denoiser = EGPULightmassDenoiser::IntelOIDN;
 
 	// Whether to compress lightmap textures.  Disabling lightmap texture compression will reduce artifacts but increase memory and disk size by 4x.
 	// Use caution when disabling this.
