Fix non unity build

diff --git a/Engine/Plugins/Experimental/GPULightmass/Source/GPULightmass/Private/LightmapDenoising.cpp b/Engine/Plugins/Experimental/GPULightmass/Source/GPULightmass/Private/LightmapDenoising.cpp
--- a/Engine/Plugins/Experimental/GPULightmass/Source/GPULightmass/Private/LightmapDenoising.cpp
+++ b/Engine/Plugins/Experimental/GPULightmass/Source/GPULightmass/Private/LightmapDenoising.cpp
@@ -2,6 +2,7 @@
 
 #include "LightmapDenoising.h"
 #include "GPULightmassCommon.h"
+#include "Async/ParallelFor.h"
 
 FDenoiserFilterSet::FDenoiserFilterSet(FDenoiserContext& Context, FIntPoint NewSize, bool bSHDenoiser)
 	: Context(Context)
