GPULM: Fix empty scene detection with MinLOD > 0 || only landscape in the scene

diff --git a/Engine/Plugins/Experimental/GPULightmass/Source/GPULightmass/Private/LightmapRenderer.cpp b/Engine/Plugins/Experimental/GPULightmass/Source/GPULightmass/Private/LightmapRenderer.cpp
--- a/Engine/Plugins/Experimental/GPULightmass/Source/GPULightmass/Private/LightmapRenderer.cpp
+++ b/Engine/Plugins/Experimental/GPULightmass/Source/GPULightmass/Private/LightmapRenderer.cpp
@@ -785,6 +785,7 @@ bool FSceneRenderState::SetupRayTracingScene(int32 LODIndex)
 		CalculateDistributionPrefixSumForAllLightmaps();
 	}
 
+	// If no LOD level is specified, select the first non empty level (and merge it with landscapes later)
 	if (LODIndex == INDEX_NONE)
 	{
 		for (int32 NonEmptyLODIndex = 0; NonEmptyLODIndex < MAX_STATIC_MESH_LODS; NonEmptyLODIndex++)
@@ -796,11 +797,6 @@ bool FSceneRenderState::SetupRayTracingScene(int32 LODIndex)
 			}
 		}
 	}
-	
-	if (CachedRayTracingScene->RayTracingGeometryInstancesPerLOD[LODIndex].Num() == 0)
-	{
-		return false;
-	}
 
 #if 0 // Debug: verify cached ray tracing scene has up-to-date shader bindings
 	TUniquePtr<FCachedRayTracingSceneData> VerificationRayTracingScene = MakeUnique<FCachedRayTracingSceneData>();
@@ -883,6 +879,12 @@ bool FSceneRenderState::SetupRayTracingScene(int32 LODIndex)
 		View.DynamicPrimitiveCollector.Commit();
 	}
 
+	// Early out if there's nothing in the scene: no instance in the selected LOD level, or no landscape (which effectively exists on every LOD level)
+	if ((LODIndex == INDEX_NONE || CachedRayTracingScene->RayTracingGeometryInstancesPerLOD[LODIndex].Num() == 0) && LandscapeRenderStates.Elements.Num() == 0)
+	{
+		return false;
+	}
+
 #if RHI_RAYTRACING
 
 	{
@@ -891,8 +893,11 @@ bool FSceneRenderState::SetupRayTracingScene(int32 LODIndex)
 		SCOPED_DRAW_EVENTF(RHICmdList, GPULightmassUpdateRayTracingScene, TEXT("GPULightmass UpdateRayTracingScene %d Instances"), StaticMeshInstanceRenderStates.Elements.Num());
 
 		TArray<FRayTracingGeometryInstance, SceneRenderingAllocator> RayTracingGeometryInstances;
-		RayTracingGeometryInstances.Append(CachedRayTracingScene->RayTracingGeometryInstancesPerLOD[LODIndex]);
-
+		if (LODIndex != INDEX_NONE)
+		{
+			RayTracingGeometryInstances.Append(CachedRayTracingScene->RayTracingGeometryInstancesPerLOD[LODIndex]);
+		}
+		
 		int32 LandscapeStartOffset = RayTracingGeometryInstances.Num();
 		for (FLandscapeRenderState& Landscape : LandscapeRenderStates.Elements)
 		{
@@ -1220,23 +1225,26 @@ bool FSceneRenderState::SetupRayTracingScene(int32 LODIndex)
 			{
 				TRACE_CPUPROFILER_EVENT_SCOPE(SetRayTracingShaderBindings);
 
-				for (const FVisibleRayTracingMeshCommand VisibleMeshCommand : CachedRayTracingScene->VisibleRayTracingMeshCommandsPerLOD[LODIndex])
+				if (LODIndex != INDEX_NONE)
 				{
-					const FRayTracingMeshCommand& MeshCommand = *VisibleMeshCommand.RayTracingMeshCommand;
-
-					MeshCommand.SetRayTracingShaderBindingsForHitGroup(BindingWriter.Get(),
-						View.ViewUniformBuffer,
-						VisibleMeshCommand.InstanceIndex,
-						MeshCommand.GeometrySegmentIndex,
-						MeshCommand.MaterialShaderIndex,
-						RAY_TRACING_SHADER_SLOT_MATERIAL);
-
-					MeshCommand.SetRayTracingShaderBindingsForHitGroup(BindingWriter.Get(),
-						View.ViewUniformBuffer,
-						VisibleMeshCommand.InstanceIndex,
-						MeshCommand.GeometrySegmentIndex,
-						MeshCommand.MaterialShaderIndex,
-						RAY_TRACING_SHADER_SLOT_SHADOW);
+					for (const FVisibleRayTracingMeshCommand VisibleMeshCommand : CachedRayTracingScene->VisibleRayTracingMeshCommandsPerLOD[LODIndex])
+					{
+						const FRayTracingMeshCommand& MeshCommand = *VisibleMeshCommand.RayTracingMeshCommand;
+
+						MeshCommand.SetRayTracingShaderBindingsForHitGroup(BindingWriter.Get(),
+							View.ViewUniformBuffer,
+							VisibleMeshCommand.InstanceIndex,
+							MeshCommand.GeometrySegmentIndex,
+							MeshCommand.MaterialShaderIndex,
+							RAY_TRACING_SHADER_SLOT_MATERIAL);
+
+						MeshCommand.SetRayTracingShaderBindingsForHitGroup(BindingWriter.Get(),
+							View.ViewUniformBuffer,
+							VisibleMeshCommand.InstanceIndex,
+							MeshCommand.GeometrySegmentIndex,
+							MeshCommand.MaterialShaderIndex,
+							RAY_TRACING_SHADER_SLOT_SHADOW);
+					}
 				}
 
 				for (const FVisibleRayTracingMeshCommand VisibleMeshCommand : VisibleRayTracingMeshCommands)
