GPULM: Clean up PrimitiveUniformShaderParameters

diff --git a/Engine/Plugins/Experimental/GPULightmass/Source/GPULightmass/Private/LightmapRenderer.cpp b/Engine/Plugins/Experimental/GPULightmass/Source/GPULightmass/Private/LightmapRenderer.cpp
--- a/Engine/Plugins/Experimental/GPULightmass/Source/GPULightmass/Private/LightmapRenderer.cpp
+++ b/Engine/Plugins/Experimental/GPULightmass/Source/GPULightmass/Private/LightmapRenderer.cpp
@@ -351,13 +351,21 @@ void FCachedRayTracingSceneData::SetupViewUniformBufferFromSceneRenderState(FSce
 			{
 				FStaticMeshInstanceRenderState& Instance = Scene.StaticMeshInstanceRenderStates.Elements[InstanceIndex];
 
-				FPrimitiveUniformShaderParameters PrimitiveUniformShaderParameters = Instance.PrimitiveUniformShaderParameters;
-				PrimitiveUniformShaderParameters.LightmapDataIndex = LightmapSceneDataStartOffsets[InstanceIndex];
-				PrimitiveUniformShaderParameters.LightmapUVIndex = 0; // TODO: LightmapUVIndex
-				PrimitiveUniformShaderParameters.InstanceSceneDataOffset = InstanceIndex;
-				PrimitiveUniformShaderParameters.NumInstanceSceneDataEntries = 1;
-				PrimitiveUniformShaderParameters.InstancePayloadDataOffset = INDEX_NONE;
-				PrimitiveUniformShaderParameters.InstancePayloadDataStride = 0;
+				FPrimitiveUniformShaderParameters PrimitiveUniformShaderParameters =
+					FPrimitiveUniformShaderParametersBuilder{}
+					.Defaults()
+					.LocalToWorld(Instance.LocalToWorld)
+					.ActorWorldPosition(Instance.ActorPosition)
+					.WorldBounds(Instance.WorldBounds)
+					.LocalBounds(Instance.LocalBounds)
+					.LightingChannelMask(0b111)
+					.LightmapDataIndex(LightmapSceneDataStartOffsets[InstanceIndex])
+					.InstanceSceneDataOffset(InstanceIndex)
+					.NumInstanceSceneDataEntries(1)
+					.InstancePayloadDataOffset(INDEX_NONE)
+					.InstancePayloadDataStride(0)
+				.Build();
+				
 				PrimitiveSceneData[InstanceIndex] = FPrimitiveSceneShaderData(PrimitiveUniformShaderParameters);
 
 				InstanceDataOriginalOffsets[InstanceIndex] = InstanceIndex;
@@ -417,8 +425,6 @@ void FCachedRayTracingSceneData::SetupViewUniformBufferFromSceneRenderState(FSce
 					.NumInstanceSceneDataEntries(NumInstancesThisGroup)
 					.InstancePayloadDataOffset(InstancePayloadData.Num())
 					.InstancePayloadDataStride(1)
-					.CastContactShadow(true)
-					.CastShadow(true)
 				.Build();
 
 			InstanceDataOriginalOffsets.Add(InstanceSceneData.Num());
@@ -485,8 +491,6 @@ void FCachedRayTracingSceneData::SetupViewUniformBufferFromSceneRenderState(FSce
 					.LightmapDataIndex(LightmapSceneDataStartOffsets[PrimitiveId])
 					.InstanceSceneDataOffset(InstanceSceneData.Num())
 					.NumInstanceSceneDataEntries(1)
-					.CastContactShadow(true)
-					.CastShadow(true)
 				.Build();
 
 			InstanceDataOriginalOffsets.Add(InstanceSceneData.Num());
diff --git a/Engine/Plugins/Experimental/GPULightmass/Source/GPULightmass/Private/Scene/Scene.cpp b/Engine/Plugins/Experimental/GPULightmass/Source/GPULightmass/Private/Scene/Scene.cpp
--- a/Engine/Plugins/Experimental/GPULightmass/Source/GPULightmass/Private/Scene/Scene.cpp
+++ b/Engine/Plugins/Experimental/GPULightmass/Source/GPULightmass/Private/Scene/Scene.cpp
@@ -797,19 +797,6 @@ void FScene::AddGeometryInstanceFromComponent(UStaticMeshComponent* InComponent)
 	{
 		FStaticMeshInstanceRenderStateRef InstanceRenderStateRef = RenderState.StaticMeshInstanceRenderStates.Emplace(MoveTemp(InstanceRenderState));
 
-		InstanceRenderStateRef->PrimitiveUniformShaderParameters =
-			FPrimitiveUniformShaderParametersBuilder{}
-			.Defaults()
-				.LocalToWorld(InstanceRenderStateRef->LocalToWorld)
-				.ActorWorldPosition(InstanceRenderStateRef->ActorPosition)
-				.WorldBounds(InstanceRenderStateRef->WorldBounds)
-				.LocalBounds(InstanceRenderStateRef->LocalBounds)
-				.LightingChannelMask(0b111)
-				.LightmapDataIndex(0)
-				.CastContactShadow(true)
-				.CastShadow(true)
-			.Build();
-
 		for (int32 LODIndex = 0; LODIndex < InstanceLightmapRenderStateInitializers.Num(); LODIndex++)
 		{
 			FLightmapRenderState::Initializer& Initializer = InstanceLightmapRenderStateInitializers[LODIndex];
diff --git a/Engine/Plugins/Experimental/GPULightmass/Source/GPULightmass/Private/Scene/StaticMesh.h b/Engine/Plugins/Experimental/GPULightmass/Source/GPULightmass/Private/Scene/StaticMesh.h
--- a/Engine/Plugins/Experimental/GPULightmass/Source/GPULightmass/Private/Scene/StaticMesh.h
+++ b/Engine/Plugins/Experimental/GPULightmass/Source/GPULightmass/Private/Scene/StaticMesh.h
@@ -19,8 +19,6 @@ public:
 	TArray<FColorVertexBuffer*> LODOverrideColorVertexBuffers;
 	TArray<TUniformBufferRef<FLocalVertexFactoryUniformShaderParameters>> LODOverrideColorVFUniformBuffers;
 
-	FPrimitiveUniformShaderParameters PrimitiveUniformShaderParameters;
-
 	TArray<FMeshBatch> GetMeshBatchesForGBufferRendering(int32 LODIndex);
 
 	int32 ClampedMinLOD;
