GPU Lightmass: Add support for source length

Make the occlusion ray calculation consistent with the latest code for ray traced shadows

diff --git a/Engine/Plugins/Experimental/GPULightmass/Shaders/Private/LightmapPathTracing.usf b/Engine/Plugins/Experimental/GPULightmass/Shaders/Private/LightmapPathTracing.usf
--- a/Engine/Plugins/Experimental/GPULightmass/Shaders/Private/LightmapPathTracing.usf
+++ b/Engine/Plugins/Experimental/GPULightmass/Shaders/Private/LightmapPathTracing.usf
@@ -15,7 +15,7 @@ PathTracingRayGenShader.usf: Reference path tracing
 #include "/Engine/Private/RayTracing/RayTracingDirectionalLight.ush"
 #include "/Engine/Private/RayTracing/RayTracingRectLight.ush"
 #include "/Engine/Private/RayTracing/RayTracingSphereLight.ush"
-#include "/Engine/Private/RayTracing/RayTracingSpotLight.ush"
+#include "/Engine/Private/RayTracing/RayTracingCapsuleLight.ush"
 #include "/Engine/Private/RayTracing/RayTracingPointLight.ush"
 
 #include "/Engine/Private/PathTracing/PathTracingCommon.ush"
@@ -1134,23 +1134,23 @@ bool GenerateOcclusionRay(
 			/* out */ RayTMin,
 			/* out */ RayTMax);
 	}
-	else if (LightType == LIGHT_TYPE_POINT)
+	else if (LightType == LIGHT_TYPE_POINT || LightType == LIGHT_TYPE_SPOT)
 	{
-		if (LightParameters.SourceRadius == 0)
+		if (LightType == LIGHT_TYPE_SPOT)
 		{
-			return GeneratePointLightOcclusionRay(
-				LightParameters,
-				TranslatedWorldPosition, WorldNormal,
-				RandSample,
-				/* out */ RayOrigin,
-				/* out */ RayDirection,
-				/* out */ RayTMin,
-				/* out */ RayTMax);
+			// before generating a shadow ray, make sure we are inside the cone of the light
+			float3 LightDirection = normalize(LightParameters.TranslatedWorldPosition - TranslatedWorldPosition);
+			float CosAngle = LightParameters.SpotAngles.x;
+			if (!(dot(LightDirection, LightParameters.Direction) >= CosAngle))
+			{
+				// outside the cone of the light, skip all work
+				return false;
+			}
 		}
-		else
+		float RayPdf;
+		if (LightParameters.SourceLength > 0.0)
 		{
-			float RayPdf;
-			return GenerateSphereLightOcclusionRayWithSolidAngleSampling(
+			return GenerateCapsuleLightOcclusionRayWithSolidAngleSampling(
 				LightParameters,
 				TranslatedWorldPosition, WorldNormal,
 				RandSample,
@@ -1160,17 +1160,15 @@ bool GenerateOcclusionRay(
 				/* out */ RayTMax,
 				/* out */ RayPdf);
 		}
-	}
-	else if (LightType == LIGHT_TYPE_SPOT)
-	{
-		return GenerateSpotLightOcclusionRay(
+		return GenerateSphereLightOcclusionRayWithSolidAngleSampling(
 			LightParameters,
 			TranslatedWorldPosition, WorldNormal,
 			RandSample,
 			/* out */ RayOrigin,
 			/* out */ RayDirection,
 			/* out */ RayTMin,
-			/* out */ RayTMax);
+			/* out */ RayTMax,
+			/* out */ RayPdf);
 	}
 	else if (LightType == LIGHT_TYPE_RECT)
 	{
