GPULM: Implement skylight bent normal denoising

diff --git a/Engine/Plugins/Experimental/GPULightmass/Source/GPULightmass/Private/LightmapDenoising.cpp b/Engine/Plugins/Experimental/GPULightmass/Source/GPULightmass/Private/LightmapDenoising.cpp
--- a/Engine/Plugins/Experimental/GPULightmass/Source/GPULightmass/Private/LightmapDenoising.cpp
+++ b/Engine/Plugins/Experimental/GPULightmass/Source/GPULightmass/Private/LightmapDenoising.cpp
@@ -62,43 +62,6 @@ FLinearColor InverseTonemapLightingForDenoising(FLinearColor Lighting)
 	return Lighting;
 }
 
-template<typename LightSampleType>
-struct LightSampleDataProvider {};
-
-template<>
-struct LightSampleDataProvider<FLinearColor>
-{
-	FIntPoint Size;
-	TArray<FLinearColor>& IncidentLighting;
-	TArray<FLinearColor>& LuminanceSH;
-
-	LightSampleDataProvider<FLinearColor>(
-		FIntPoint Size,
-		TArray<FLinearColor>& IncidentLighting,
-		TArray<FLinearColor>& LuminanceSH)
-		: Size(Size)
-		, IncidentLighting(IncidentLighting)
-		, LuminanceSH(LuminanceSH)
-	{
-	}
-
-	bool IsMapped(FIntPoint Location)
-	{
-		return IncidentLighting[Location.Y * Size.X + Location.X].A >= 0;
-	}
-	
-	float GetL(FIntPoint Location)
-	{
-		return IncidentLighting[Location.Y * Size.X + Location.X].GetLuminance();
-	}
-
-	void OverwriteTexel(FIntPoint Dst, FIntPoint Src)
-	{
-		IncidentLighting[Dst.Y * Size.X + Dst.X] = IncidentLighting[Src.Y * Size.X + Src.X];
-		LuminanceSH[Dst.Y * Size.X + Dst.X] = LuminanceSH[Src.Y * Size.X + Src.X];
-	}
-};
-
 const int32 WindowSize = 2;
 
 int32 EncodeOffset(FIntPoint D)
@@ -120,12 +83,9 @@ bool IsColorValid(FLinearColor Color)
 	&& Color.R >= 0 && Color.G >= 0 && Color.B >= 0 && Color.A >= 0;
 }
 
-void SimpleFireflyFilter(
-	FIntPoint Size,
-	TArray<FLinearColor>& IncidentLighting,
-	TArray<FLinearColor>& LuminanceSH)
+void SimpleFireflyFilter(FLightSampleDataProvider& SampleData)
 {
-	LightSampleDataProvider<FLinearColor> SampleData(Size, IncidentLighting, LuminanceSH);
+	const FIntPoint Size = SampleData.GetSize();
 	
 	TArray<int32> Buffer;
 	Buffer.AddDefaulted(Size.X * Size.Y);
@@ -206,9 +166,15 @@ void NormalizeDirectionalLuminancePair(
 				DirLuma[2] = LuminanceSH[Y * Size.X + X].G;	// Keep the rest as we need to diffuse conv them anyway
 				DirLuma[3] = LuminanceSH[Y * Size.X + X].B;	// Keep the rest as we need to diffuse conv them anyway
 
-				float DirScale = 1.0f / FMath::Max(0.0001f, DirLuma[0]);
+				float DirScale = 1.0f / DirLuma[0];
 				float ColorScale = DirLuma[0];
 
+				if (DirLuma[0] < 0.00001f)
+				{
+					DirScale = 0.0f;
+					ColorScale = 0.0f;
+				}
+
 				{
 					IncidentLighting[Y * Size.X + X].R = ColorScale * IncidentLighting[Y * Size.X + X].R * IncidentLighting[Y * Size.X + X].R;
 					IncidentLighting[Y * Size.X + X].G = ColorScale * IncidentLighting[Y * Size.X + X].G * IncidentLighting[Y * Size.X + X].G;
@@ -376,6 +342,88 @@ void DenoiseDirectionalLuminancePair(
 	}
 }
 
+void DenoiseSkyBentNormal(
+	FIntPoint Size,
+	TArray<FLinearColor>& ValidityMask,
+	TArray<FVector3f>& BentNormal,
+	FDenoiserContext& DenoiserContext)
+{
+	TArray<int32> DilationMask;
+	DilationMask.AddZeroed(Size.X * Size.Y);
+
+	for (int32 Y = 0; Y < Size.Y; Y++)
+	{
+		for (int32 X = 0; X < Size.X; X++)
+		{
+			DilationMask[Y * Size.X + X] = ValidityMask[Y * Size.X + X].A >= 0.0f ? 1 : 0;
+		}
+	}
+	
+	const int32 TotalIteration = 2;
+
+	for (int32 IterationIndex = 0; IterationIndex < TotalIteration; IterationIndex++)
+	{
+		for (int32 Y = 0; Y < Size.Y; Y++)
+		{
+			for (int32 X = 0; X < Size.X; X++)
+			{
+				if (DilationMask[Y * Size.X + X] == 0)
+				{
+					for (int32 dx = -1; dx <= 1; dx++)
+					{
+						for (int32 dy = -1; dy <= 1; dy++)
+						{
+							if (Y + dy >= 0 && Y + dy < Size.Y && X + dx >= 0 && X + dx < Size.X)
+							{
+								if (DilationMask[(Y + dy) * Size.X + (X + dx)] == 1)
+								{
+									BentNormal[Y * Size.X + X] = BentNormal[(Y + dy) * Size.X + (X + dx)];
+									DilationMask[Y * Size.X + X] = 2;
+								}
+							}
+						}
+					}
+				}
+			}
+		}
+
+		for (int32 Y = 0; Y < Size.Y; Y++)
+		{
+			for (int32 X = 0; X < Size.X; X++)
+			{
+				if (DilationMask[Y * Size.X + X] == 2)
+				{
+					DilationMask[Y * Size.X + X] = 1;
+				}
+			}
+		}
+	}
+	
+	int32 Padding = 8;
+	
+	// Resizing the filter is a super expensive operation
+	// Round things into size bins to reduce number of resizes
+	FDenoiserFilterSet& FilterSet = DenoiserContext.GetFilterForSize(FIntPoint(FMath::DivideAndRoundUp(Size.X + Padding * 2, 64) * 64, FMath::DivideAndRoundUp(Size.Y + Padding * 2, 64) * 64), true);
+	
+	for (int32 Y = 0; Y < Size.Y; Y++)
+	{
+		for (int32 X = 0; X < Size.X; X++)
+		{
+			FilterSet.InputBuffer[(Y + Padding) * FilterSet.Size.X + (X + Padding)] = BentNormal[Y * Size.X + X];
+		}
+	}
+
+	FilterSet.Execute();
+
+	for (int32 Y = 0; Y < Size.Y; Y++)
+	{
+		for (int32 X = 0; X < Size.X; X++)
+		{
+			BentNormal[Y * Size.X + X] = FilterSet.OutputBuffer[(Y + Padding) * FilterSet.Size.X + (X + Padding)];
+		}
+	}
+}
+
 void DenoiseRawData(
 	FIntPoint Size,
 	TArray<FLinearColor>& IncidentLighting,
diff --git a/Engine/Plugins/Experimental/GPULightmass/Source/GPULightmass/Private/LightmapDenoising.h b/Engine/Plugins/Experimental/GPULightmass/Source/GPULightmass/Private/LightmapDenoising.h
--- a/Engine/Plugins/Experimental/GPULightmass/Source/GPULightmass/Private/LightmapDenoising.h
+++ b/Engine/Plugins/Experimental/GPULightmass/Source/GPULightmass/Private/LightmapDenoising.h
@@ -88,6 +88,12 @@ struct FDenoiserContext
 	}
 };
 
+void DenoiseSkyBentNormal(
+	FIntPoint Size,
+	TArray<FLinearColor>& ValidityMask,
+	TArray<FVector3f>& BentNormal,
+	FDenoiserContext& DenoiserContext);
+
 void DenoiseRawData(
 	FIntPoint Size,
 	TArray<FLinearColor>& IncidentLighting,
@@ -95,7 +101,93 @@ void DenoiseRawData(
 	FDenoiserContext& DenoiserContext,
 	bool bPrepadTexels = true);
 
-void SimpleFireflyFilter(
-	FIntPoint Size,
-	TArray<FLinearColor>& IncidentLighting,
-	TArray<FLinearColor>& LuminanceSH);
+struct FLightSampleDataProvider
+{
+	virtual ~FLightSampleDataProvider() {}
+	virtual FIntPoint GetSize() const = 0; 
+	virtual bool IsMapped(FIntPoint Location) const = 0;	
+	virtual float GetL(FIntPoint Location) const = 0;
+	virtual void OverwriteTexel(FIntPoint Dst, FIntPoint Src) = 0;
+};
+
+template<typename LightSampleType>
+struct TLightSampleDataProvider {};
+
+template<>
+struct TLightSampleDataProvider<FLinearColor> : public FLightSampleDataProvider
+{
+	FIntPoint Size;
+	TArray<FLinearColor>& IncidentLighting;
+	TArray<FLinearColor>& LuminanceSH;
+
+	TLightSampleDataProvider<FLinearColor>(
+		FIntPoint Size,
+		TArray<FLinearColor>& IncidentLighting,
+		TArray<FLinearColor>& LuminanceSH)
+		: Size(Size)
+		, IncidentLighting(IncidentLighting)
+		, LuminanceSH(LuminanceSH)
+	{
+	}
+	
+	virtual FIntPoint GetSize() const override
+	{
+		return Size;
+	}
+
+	virtual bool IsMapped(FIntPoint Location) const override
+	{
+		return IncidentLighting[Location.Y * Size.X + Location.X].A >= 0;
+	}
+	
+	virtual float GetL(FIntPoint Location) const override
+	{
+		return IncidentLighting[Location.Y * Size.X + Location.X].GetLuminance();
+	}
+
+	virtual void OverwriteTexel(FIntPoint Dst, FIntPoint Src) override
+	{
+		IncidentLighting[Dst.Y * Size.X + Dst.X] = IncidentLighting[Src.Y * Size.X + Src.X];
+		LuminanceSH[Dst.Y * Size.X + Dst.X] = LuminanceSH[Src.Y * Size.X + Src.X];
+	}
+};
+
+template<>
+struct TLightSampleDataProvider<FVector3f> : public FLightSampleDataProvider
+{
+	FIntPoint Size;
+	TArray<FLinearColor>& IncidentLightingAsValidityMask;
+	TArray<FVector3f>& SkyBentNormal;
+
+	TLightSampleDataProvider<FVector3f>(
+		FIntPoint Size,
+		TArray<FLinearColor>& IncidentLightingAsValidityMask,
+		TArray<FVector3f>& SkyBentNormal)
+		: Size(Size)
+		, IncidentLightingAsValidityMask(IncidentLightingAsValidityMask)
+		, SkyBentNormal(SkyBentNormal)
+	{
+	}
+	
+	virtual FIntPoint GetSize() const override
+	{
+		return Size;
+	}
+
+	virtual bool IsMapped(FIntPoint Location) const override
+	{
+		return IncidentLightingAsValidityMask[Location.Y * Size.X + Location.X].A >= 0;
+	}
+	
+	virtual float GetL(FIntPoint Location) const override
+	{
+		return SkyBentNormal[Location.Y * Size.X + Location.X].Length();
+	}
+
+	virtual void OverwriteTexel(FIntPoint Dst, FIntPoint Src) override
+	{
+		SkyBentNormal[Dst.Y * Size.X + Dst.X] = SkyBentNormal[Src.Y * Size.X + Src.X];
+	}
+};
+
+void SimpleFireflyFilter(FLightSampleDataProvider& SampleData);
diff --git a/Engine/Plugins/Experimental/GPULightmass/Source/GPULightmass/Private/LightmapEncoding.cpp b/Engine/Plugins/Experimental/GPULightmass/Source/GPULightmass/Private/LightmapEncoding.cpp
--- a/Engine/Plugins/Experimental/GPULightmass/Source/GPULightmass/Private/LightmapEncoding.cpp
+++ b/Engine/Plugins/Experimental/GPULightmass/Source/GPULightmass/Private/LightmapEncoding.cpp
@@ -347,16 +347,6 @@ FLightSampleData ConvertToLightSample(FLinearColor IncidentLighting, FLinearColo
 	return Sample;
 }
 
-void AddSkyOcclusionToLightSample(FLightSampleData& Sample, FLinearColor SkyOcclusion)
-{
-	// Revert sqrt in LightmapEncoding.ush for preview
-	float Length = SkyOcclusion.A * SkyOcclusion.A;
-	FVector3f UnpackedBentNormalVector = FVector3f(SkyOcclusion) * 2.0f - 1.0f;
-	Sample.SkyOcclusion[0] = UnpackedBentNormalVector.X * Length;
-	Sample.SkyOcclusion[1] = UnpackedBentNormalVector.Y * Length;
-	Sample.SkyOcclusion[2] = UnpackedBentNormalVector.Z * Length;
-}
-
 FQuantizedSignedDistanceFieldShadowSample ConvertToShadowSample(FLinearColor ShadowMask, int32 ChannelIndex)
 {
 	FQuantizedSignedDistanceFieldShadowSample Sample;
diff --git a/Engine/Plugins/Experimental/GPULightmass/Source/GPULightmass/Private/LightmapEncoding.h b/Engine/Plugins/Experimental/GPULightmass/Source/GPULightmass/Private/LightmapEncoding.h
--- a/Engine/Plugins/Experimental/GPULightmass/Source/GPULightmass/Private/LightmapEncoding.h
+++ b/Engine/Plugins/Experimental/GPULightmass/Source/GPULightmass/Private/LightmapEncoding.h
@@ -51,7 +51,6 @@ struct FLightSampleData
 };
 
 FLightSampleData ConvertToLightSample(FLinearColor IncidentLighting, FLinearColor LuminanceSH);
-void AddSkyOcclusionToLightSample(FLightSampleData& Sample, FLinearColor SkyOcclusion);
 
 void QuantizeLightSamples(
 	TArray<FLightSampleData> InLightSamples,
diff --git a/Engine/Plugins/Experimental/GPULightmass/Source/GPULightmass/Private/LightmapRenderer.cpp b/Engine/Plugins/Experimental/GPULightmass/Source/GPULightmass/Private/LightmapRenderer.cpp
--- a/Engine/Plugins/Experimental/GPULightmass/Source/GPULightmass/Private/LightmapRenderer.cpp
+++ b/Engine/Plugins/Experimental/GPULightmass/Source/GPULightmass/Private/LightmapRenderer.cpp
@@ -3044,9 +3044,31 @@ const int32 DenoiseTileProximity = 3;
 
 void FLightmapTileDenoiseAsyncTask::DoThreadedWork()
 {
+	TArray<FVector3f> SkyBentNormal;
+	SkyBentNormal.AddZeroed(Size.X * Size.Y);
+		
+	for (int32 Y = 0 ; Y < Size.Y; Y++)
+	{
+		for (int32 X = 0 ; X < Size.X; X++)
+		{
+			int32 LinearIndex = Y * Size.X + X;
+										
+			FLinearColor SkyOcclusion = TextureData->Texture[3][LinearIndex];
+										
+			// Revert sqrt in LightmapEncoding.ush for preview
+			float Length = SkyOcclusion.A * SkyOcclusion.A;
+			FVector3f UnpackedBentNormalVector = FVector3f(SkyOcclusion) * 2.0f - 1.0f;
+			SkyBentNormal[LinearIndex] = UnpackedBentNormalVector * Length;
+		}
+	}
+	
 	if (Denoiser == EGPULightmassDenoiser::SimpleFireflyRemover)
 	{
-		SimpleFireflyFilter(Size, TextureData->Texture[0], TextureData->Texture[1]);
+		TLightSampleDataProvider<FLinearColor> GISampleData(Size, TextureData->Texture[0], TextureData->Texture[1]);
+		SimpleFireflyFilter(GISampleData);
+
+		TLightSampleDataProvider<FVector3f> SkyBentNormalSampleData(Size, TextureData->Texture[0], SkyBentNormal);
+		SimpleFireflyFilter(SkyBentNormalSampleData);
 	}
 	else
 	{
@@ -3057,6 +3079,24 @@ void FLightmapTileDenoiseAsyncTask::DoThreadedWork()
 			TextureData->Texture[0],
 			TextureData->Texture[1],
 			DenoiserContext);
+
+		DenoiseSkyBentNormal(Size, TextureData->Texture[0], SkyBentNormal, DenoiserContext);
+	}
+		
+	for (int32 Y = 0 ; Y < Size.Y; Y++)
+	{
+		for (int32 X = 0 ; X < Size.X; X++)
+		{
+			int32 LinearIndex = Y * Size.X + X;
+				
+			float Length = SkyBentNormal[LinearIndex].Length();
+			FVector3f PackedVector = SkyBentNormal[LinearIndex].GetSafeNormal() * 0.5f + 0.5f;
+				
+			TextureData->Texture[3][LinearIndex].R = PackedVector.X;
+			TextureData->Texture[3][LinearIndex].G = PackedVector.Y;
+			TextureData->Texture[3][LinearIndex].B = PackedVector.Z;
+			TextureData->Texture[3][LinearIndex].A = FMath::Sqrt(Length);
+		}
 	}
 
 	FPlatformAtomics::AtomicStore(&TextureData->bDenoisingFinished, 1);
@@ -3107,6 +3147,8 @@ void FLightmapRenderer::BackgroundTick()
 				// While the data will be overwritten immediately, we still need to decompress to inform the LRU cache management
 				Tile.RenderState->TileStorage[Tile.VirtualCoordinates].CPUTextureData[0]->Decompress();
 				Tile.RenderState->TileStorage[Tile.VirtualCoordinates].CPUTextureData[1]->Decompress();
+				Tile.RenderState->TileStorage[Tile.VirtualCoordinates].CPUTextureData[2]->Decompress();
+				Tile.RenderState->TileStorage[Tile.VirtualCoordinates].CPUTextureData[3]->Decompress();
 
 				for (int32 Y = 0; Y < GPreviewLightmapVirtualTileSize; Y++)
 				{
@@ -3120,6 +3162,8 @@ void FLightmapRenderer::BackgroundTick()
 
 						Tile.RenderState->TileStorage[Tile.VirtualCoordinates].CPUTextureData[0]->Data[DstLinearIndex] = DenoiseGroup.TextureData->Texture[0][SrcLinearIndex];
 						Tile.RenderState->TileStorage[Tile.VirtualCoordinates].CPUTextureData[1]->Data[DstLinearIndex] = DenoiseGroup.TextureData->Texture[1][SrcLinearIndex];
+						Tile.RenderState->TileStorage[Tile.VirtualCoordinates].CPUTextureData[2]->Data[DstLinearIndex] = DenoiseGroup.TextureData->Texture[2][SrcLinearIndex];
+						Tile.RenderState->TileStorage[Tile.VirtualCoordinates].CPUTextureData[3]->Data[DstLinearIndex] = DenoiseGroup.TextureData->Texture[3][SrcLinearIndex];
 					}
 				}
 
@@ -3211,6 +3255,8 @@ void FLightmapRenderer::BackgroundTick()
 				{
 					TileStorage.CPUTextureRawData[0]->Decompress();
 					TileStorage.CPUTextureRawData[1]->Decompress();
+					TileStorage.CPUTextureRawData[2]->Decompress();
+					TileStorage.CPUTextureRawData[3]->Decompress();
 				}
 
 				TileStorage.CPUTextureData[0]->Decompress();
@@ -3232,6 +3278,8 @@ void FLightmapRenderer::BackgroundTick()
 						{
 							TileStorage.CPUTextureRawData[0]->Data[DstLinearIndex] = ReadbackGroup.TextureData->Texture[0][SrcLinearIndex];
 							TileStorage.CPUTextureRawData[1]->Data[DstLinearIndex] = ReadbackGroup.TextureData->Texture[1][SrcLinearIndex];
+							TileStorage.CPUTextureRawData[2]->Data[DstLinearIndex] = ReadbackGroup.TextureData->Texture[2][SrcLinearIndex];
+							TileStorage.CPUTextureRawData[3]->Data[DstLinearIndex] = ReadbackGroup.TextureData->Texture[3][SrcLinearIndex];
 						}
 
 						// Always write into display data so we have something to show before denoising completes
@@ -3365,6 +3413,8 @@ void FLightmapRenderer::BackgroundTick()
 
 			DenoiseGroup.TextureData->Texture[0].AddZeroed(DenoiseTileProximity * DenoiseTileProximity * GPreviewLightmapVirtualTileSize * GPreviewLightmapVirtualTileSize);
 			DenoiseGroup.TextureData->Texture[1].AddZeroed(DenoiseTileProximity * DenoiseTileProximity * GPreviewLightmapVirtualTileSize * GPreviewLightmapVirtualTileSize);
+			DenoiseGroup.TextureData->Texture[2].AddZeroed(DenoiseTileProximity * DenoiseTileProximity * GPreviewLightmapVirtualTileSize * GPreviewLightmapVirtualTileSize);
+			DenoiseGroup.TextureData->Texture[3].AddZeroed(DenoiseTileProximity * DenoiseTileProximity * GPreviewLightmapVirtualTileSize * GPreviewLightmapVirtualTileSize);
 
 			for (int Dx = -(DenoiseTileProximity / 2); Dx <= (DenoiseTileProximity / 2); Dx++)
 			{
@@ -3389,6 +3439,8 @@ void FLightmapRenderer::BackgroundTick()
 					{
 						Tile.RenderState->TileStorage[FTileVirtualCoordinates(SrcTilePosition, Tile.VirtualCoordinates.MipLevel)].CPUTextureRawData[0]->Decompress();
 						Tile.RenderState->TileStorage[FTileVirtualCoordinates(SrcTilePosition, Tile.VirtualCoordinates.MipLevel)].CPUTextureRawData[1]->Decompress();
+						Tile.RenderState->TileStorage[FTileVirtualCoordinates(SrcTilePosition, Tile.VirtualCoordinates.MipLevel)].CPUTextureRawData[2]->Decompress();
+						Tile.RenderState->TileStorage[FTileVirtualCoordinates(SrcTilePosition, Tile.VirtualCoordinates.MipLevel)].CPUTextureRawData[3]->Decompress();
 					}
 
 					for (int32 Y = 0; Y < GPreviewLightmapVirtualTileSize; Y++)
@@ -3403,6 +3455,8 @@ void FLightmapRenderer::BackgroundTick()
 
 							DenoiseGroup.TextureData->Texture[0][DstLinearIndex] = !bShouldWriteZero ? Tile.RenderState->TileStorage[FTileVirtualCoordinates(SrcTilePosition, Tile.VirtualCoordinates.MipLevel)].CPUTextureRawData[0]->Data[SrcLinearIndex] : FLinearColor(0, 0, 0, 0);
 							DenoiseGroup.TextureData->Texture[1][DstLinearIndex] = !bShouldWriteZero ? Tile.RenderState->TileStorage[FTileVirtualCoordinates(SrcTilePosition, Tile.VirtualCoordinates.MipLevel)].CPUTextureRawData[1]->Data[SrcLinearIndex] : FLinearColor(0, 0, 0, 0);
+							DenoiseGroup.TextureData->Texture[2][DstLinearIndex] = !bShouldWriteZero ? Tile.RenderState->TileStorage[FTileVirtualCoordinates(SrcTilePosition, Tile.VirtualCoordinates.MipLevel)].CPUTextureRawData[2]->Data[SrcLinearIndex] : FLinearColor(0, 0, 0, 0);
+							DenoiseGroup.TextureData->Texture[3][DstLinearIndex] = !bShouldWriteZero ? Tile.RenderState->TileStorage[FTileVirtualCoordinates(SrcTilePosition, Tile.VirtualCoordinates.MipLevel)].CPUTextureRawData[3]->Data[SrcLinearIndex] : FLinearColor(0, 0, 0, 0);
 						}
 					}
 				}
diff --git a/Engine/Plugins/Experimental/GPULightmass/Source/GPULightmass/Private/LightmapRenderer.h b/Engine/Plugins/Experimental/GPULightmass/Source/GPULightmass/Private/LightmapRenderer.h
--- a/Engine/Plugins/Experimental/GPULightmass/Source/GPULightmass/Private/LightmapRenderer.h
+++ b/Engine/Plugins/Experimental/GPULightmass/Source/GPULightmass/Private/LightmapRenderer.h
@@ -86,7 +86,7 @@ struct FLightmapTileDenoiseGroup
 
 	struct FTextureData
 	{
-		TArray<FLinearColor> Texture[2];
+		TArray<FLinearColor> Texture[4];
 
 		volatile int32 bDenoisingFinished = 0;
 	};
diff --git a/Engine/Plugins/Experimental/GPULightmass/Source/GPULightmass/Private/Scene/Scene.cpp b/Engine/Plugins/Experimental/GPULightmass/Source/GPULightmass/Private/Scene/Scene.cpp
--- a/Engine/Plugins/Experimental/GPULightmass/Source/GPULightmass/Private/Scene/Scene.cpp
+++ b/Engine/Plugins/Experimental/GPULightmass/Source/GPULightmass/Private/Scene/Scene.cpp
@@ -2033,6 +2033,8 @@ void FScene::ApplyFinishedLightmapsToWorld()
 
 						TArray<FLinearColor> IncidentLighting;
 						TArray<FLinearColor> LuminanceSH;
+						TArray<FVector3f> SkyBentNormal;
+						
 						IncidentLighting.AddZeroed(Lightmap.GetSize().X * Lightmap.GetSize().Y);
 						LuminanceSH.AddZeroed(Lightmap.GetSize().X * Lightmap.GetSize().Y);
 						
@@ -2051,39 +2053,71 @@ void FScene::ApplyFinishedLightmapsToWorld()
 								LuminanceSH[DstLinearIndex] = Lightmap.TileStorage[FTileVirtualCoordinates(SrcTilePosition, 0)].CPUTextureData[1]->Data[SrcLinearIndex];
 							});
 
-							if (Settings->DenoisingOptions == EGPULightmassDenoisingOptions::OnCompletion)
+							if (bHasSkyShadowing)
 							{
-								if (Settings->Denoiser == EGPULightmassDenoiser::SimpleFireflyRemover)
+								SkyBentNormal.AddZeroed(Lightmap.GetSize().X * Lightmap.GetSize().Y);
+
+								CopyRectTiled(
+									FIntPoint(0, 0),
+									FIntRect(FIntPoint(0, 0), Lightmap.GetSize()),
+									SrcRowPitchInPixels,
+									DstRowPitchInPixels,
+									[&Lightmap, &SkyBentNormal](int32 DstLinearIndex, FIntPoint SrcTilePosition, int32 SrcLinearIndex) mutable
+									{
+										FLinearColor SkyOcclusion = Lightmap.TileStorage[FTileVirtualCoordinates(SrcTilePosition, 0)].CPUTextureData[3]->Data[SrcLinearIndex];
+										
+										// Revert sqrt in LightmapEncoding.ush which was done for preview
+										float Length = SkyOcclusion.A * SkyOcclusion.A;
+										FVector3f UnpackedBentNormalVector = FVector3f(SkyOcclusion) * 2.0f - 1.0f;
+										SkyBentNormal[DstLinearIndex] = UnpackedBentNormalVector * Length;
+									});
+							}
+						}
+
+						if (Settings->DenoisingOptions == EGPULightmassDenoisingOptions::OnCompletion)
+						{
+							if (Settings->Denoiser == EGPULightmassDenoiser::SimpleFireflyRemover)
+							{
+								TLightSampleDataProvider<FLinearColor> GISampleData(Lightmap.GetSize(), IncidentLighting, LuminanceSH);
+								SimpleFireflyFilter(GISampleData);
+
+								if (bHasSkyShadowing)
 								{
-									SimpleFireflyFilter(Lightmap.GetSize(), IncidentLighting, LuminanceSH);
+									TLightSampleDataProvider<FVector3f> SkyBentNormalSampleData(Lightmap.GetSize(), IncidentLighting, SkyBentNormal);
+									SimpleFireflyFilter(SkyBentNormalSampleData);
 								}
-								else
+							}
+							else
+							{
+								if (bHasSkyShadowing)
 								{
-									DenoiseRawData(Lightmap.GetSize(), IncidentLighting, LuminanceSH, DenoiserContext);
+									DenoiseSkyBentNormal(Lightmap.GetSize(), IncidentLighting, SkyBentNormal, DenoiserContext);
 								}
+								DenoiseRawData(Lightmap.GetSize(), IncidentLighting, LuminanceSH, DenoiserContext);
 							}
+						}
 							
+						for (int32 Y = 0 ; Y < Lightmap.GetSize().Y; Y++)
+						{
+							for (int32 X = 0 ; X < Lightmap.GetSize().X; X++)
+							{
+								int32 LinearIndex = Y * Lightmap.GetSize().X + X;
+								LightSampleData[LinearIndex] = ConvertToLightSample(IncidentLighting[LinearIndex], LuminanceSH[LinearIndex]);
+							}
+						}
+								
+						if (bHasSkyShadowing)
+						{
 							for (int32 Y = 0 ; Y < Lightmap.GetSize().Y; Y++)
 							{
 								for (int32 X = 0 ; X < Lightmap.GetSize().X; X++)
 								{
 									int32 LinearIndex = Y * Lightmap.GetSize().X + X;
-									LightSampleData[LinearIndex] = ConvertToLightSample(IncidentLighting[LinearIndex], LuminanceSH[LinearIndex]);
+									LightSampleData[LinearIndex].SkyOcclusion[0] = SkyBentNormal[LinearIndex].X;
+									LightSampleData[LinearIndex].SkyOcclusion[1] = SkyBentNormal[LinearIndex].Y;
+									LightSampleData[LinearIndex].SkyOcclusion[2] = SkyBentNormal[LinearIndex].Z;
 								}
 							}
-
-							if (bHasSkyShadowing)
-							{
-								CopyRectTiled(
-									FIntPoint(0, 0),
-									FIntRect(FIntPoint(0, 0), Lightmap.GetSize()),
-									SrcRowPitchInPixels,
-									DstRowPitchInPixels,
-									[&Lightmap, &LightSampleData](int32 DstLinearIndex, FIntPoint SrcTilePosition, int32 SrcLinearIndex) mutable
-								{
-									AddSkyOcclusionToLightSample(LightSampleData[DstLinearIndex], Lightmap.TileStorage[FTileVirtualCoordinates(SrcTilePosition, 0)].CPUTextureData[3]->Data[SrcLinearIndex]);
-								});
-							}
 						}
 
 #if 0
@@ -2337,6 +2371,7 @@ void FScene::ApplyFinishedLightmapsToWorld()
 						{
 							TArray<FLinearColor> IncidentLighting;
 							TArray<FLinearColor> LuminanceSH;
+							TArray<FVector3f> SkyBentNormal;
 							IncidentLighting.AddZeroed(BaseLightMapWidth * BaseLightMapHeight);
 							LuminanceSH.AddZeroed(BaseLightMapWidth * BaseLightMapHeight);
 							
@@ -2364,15 +2399,46 @@ void FScene::ApplyFinishedLightmapsToWorld()
 									IncidentLighting[DstLinearIndex] = Lightmap.TileStorage[FTileVirtualCoordinates(SrcTilePosition, 0)].CPUTextureData[0]->Data[SrcLinearIndex];
 									LuminanceSH[DstLinearIndex] = Lightmap.TileStorage[FTileVirtualCoordinates(SrcTilePosition, 0)].CPUTextureData[1]->Data[SrcLinearIndex];
 								});
+
+								if (bHasSkyShadowing)
+								{
+									SkyBentNormal.AddZeroed(BaseLightMapWidth * BaseLightMapHeight);
+									
+									CopyRectTiled(
+										InstanceTileMin,
+										FIntRect(FIntPoint(0, 0), FIntPoint(BaseLightMapWidth, BaseLightMapHeight)),
+										SrcRowPitchInPixels,
+										DstRowPitchInPixels,
+										[&Lightmap, &SkyBentNormal](int32 DstLinearIndex, FIntPoint SrcTilePosition, int32 SrcLinearIndex) mutable
+										{
+											FLinearColor SkyOcclusion = Lightmap.TileStorage[FTileVirtualCoordinates(SrcTilePosition, 0)].CPUTextureData[3]->Data[SrcLinearIndex];
+
+											// Revert sqrt in LightmapEncoding.ush which was done for preview
+											float Length = SkyOcclusion.A * SkyOcclusion.A;
+											FVector3f UnpackedBentNormalVector = FVector3f(SkyOcclusion) * 2.0f - 1.0f;
+											SkyBentNormal[DstLinearIndex] = UnpackedBentNormalVector * Length;
+										});
+								}
 								
 								if (Settings->DenoisingOptions == EGPULightmassDenoisingOptions::OnCompletion)
 								{
 									if (Settings->Denoiser == EGPULightmassDenoiser::SimpleFireflyRemover)
 									{
-										SimpleFireflyFilter(FIntPoint{BaseLightMapWidth, BaseLightMapHeight}, IncidentLighting, LuminanceSH);
+										TLightSampleDataProvider<FLinearColor> SampleData(FIntPoint{BaseLightMapWidth, BaseLightMapHeight}, IncidentLighting, LuminanceSH);
+										SimpleFireflyFilter(SampleData);
+										
+										if (bHasSkyShadowing)
+										{
+											TLightSampleDataProvider<FVector3f> SkyBentNormalSampleData(FIntPoint{BaseLightMapWidth, BaseLightMapHeight}, IncidentLighting, SkyBentNormal);
+											SimpleFireflyFilter(SkyBentNormalSampleData);
+										}
 									}
 									else
 									{
+										if (bHasSkyShadowing)
+										{
+											DenoiseSkyBentNormal(FIntPoint{BaseLightMapWidth, BaseLightMapHeight}, IncidentLighting, SkyBentNormal, DenoiserContext);
+										}
 										DenoiseRawData(FIntPoint{BaseLightMapWidth, BaseLightMapHeight}, IncidentLighting, LuminanceSH, DenoiserContext);
 									}
 								}
@@ -2388,15 +2454,16 @@ void FScene::ApplyFinishedLightmapsToWorld()
 
 								if (bHasSkyShadowing)
 								{
-									CopyRectTiled(
-										InstanceTileMin,
-										FIntRect(FIntPoint(0, 0), FIntPoint(BaseLightMapWidth, BaseLightMapHeight)),
-										SrcRowPitchInPixels,
-										DstRowPitchInPixels,
-										[&Lightmap, &LightSampleData](int32 DstLinearIndex, FIntPoint SrcTilePosition, int32 SrcLinearIndex) mutable
+									for (int32 Y = 0; Y < BaseLightMapHeight; Y++)
 									{
-										AddSkyOcclusionToLightSample(LightSampleData[DstLinearIndex], Lightmap.TileStorage[FTileVirtualCoordinates(SrcTilePosition, 0)].CPUTextureData[3]->Data[SrcLinearIndex]);
-									});
+										for (int32 X = 0; X < BaseLightMapWidth; X++)
+										{
+											int32 LinearIndex = Y * BaseLightMapWidth + X;
+											LightSampleData[LinearIndex].SkyOcclusion[0] = SkyBentNormal[LinearIndex].X;
+											LightSampleData[LinearIndex].SkyOcclusion[1] = SkyBentNormal[LinearIndex].Y;
+											LightSampleData[LinearIndex].SkyOcclusion[2] = SkyBentNormal[LinearIndex].Z;
+										}
+									}
 								}
 							}
 
@@ -2632,6 +2699,7 @@ void FScene::ApplyFinishedLightmapsToWorld()
 
 						TArray<FLinearColor> IncidentLighting;
 						TArray<FLinearColor> LuminanceSH;
+						TArray<FVector3f> SkyBentNormal;
 						IncidentLighting.AddZeroed(Lightmap.GetSize().X * Lightmap.GetSize().Y);
 						LuminanceSH.AddZeroed(Lightmap.GetSize().X * Lightmap.GetSize().Y);
 						
@@ -2649,15 +2717,46 @@ void FScene::ApplyFinishedLightmapsToWorld()
 								IncidentLighting[DstLinearIndex] = Lightmap.TileStorage[FTileVirtualCoordinates(SrcTilePosition, 0)].CPUTextureData[0]->Data[SrcLinearIndex];
 								LuminanceSH[DstLinearIndex] = Lightmap.TileStorage[FTileVirtualCoordinates(SrcTilePosition, 0)].CPUTextureData[1]->Data[SrcLinearIndex];
 							});
+							
+							if (bHasSkyShadowing)
+							{
+								SkyBentNormal.AddZeroed(Lightmap.GetSize().X * Lightmap.GetSize().Y);
+
+								CopyRectTiled(
+									FIntPoint(0, 0),
+									FIntRect(FIntPoint(0, 0), Lightmap.GetSize()),
+									SrcRowPitchInPixels,
+									DstRowPitchInPixels,
+									[&Lightmap, &SkyBentNormal](int32 DstLinearIndex, FIntPoint SrcTilePosition, int32 SrcLinearIndex) mutable
+									{
+										FLinearColor SkyOcclusion = Lightmap.TileStorage[FTileVirtualCoordinates(SrcTilePosition, 0)].CPUTextureData[3]->Data[SrcLinearIndex];
+										
+										// Revert sqrt in LightmapEncoding.ush which was done for preview
+										float Length = SkyOcclusion.A * SkyOcclusion.A;
+										FVector3f UnpackedBentNormalVector = FVector3f(SkyOcclusion) * 2.0f - 1.0f;
+										SkyBentNormal[DstLinearIndex] = UnpackedBentNormalVector * Length;
+									});
+							}
 
 							if (Settings->DenoisingOptions == EGPULightmassDenoisingOptions::OnCompletion)
 							{
 								if (Settings->Denoiser == EGPULightmassDenoiser::SimpleFireflyRemover)
 								{
-									SimpleFireflyFilter(Lightmap.GetSize(), IncidentLighting, LuminanceSH);
+									TLightSampleDataProvider<FLinearColor> GISampleData(Lightmap.GetSize(), IncidentLighting, LuminanceSH);
+									SimpleFireflyFilter(GISampleData);
+
+									if (bHasSkyShadowing)
+									{
+										TLightSampleDataProvider<FVector3f> SkyBentNormalSampleData(Lightmap.GetSize(), IncidentLighting, SkyBentNormal);
+										SimpleFireflyFilter(SkyBentNormalSampleData);
+									}
 								}
 								else
 								{
+									if (bHasSkyShadowing)
+									{
+										DenoiseSkyBentNormal(Lightmap.GetSize(), IncidentLighting, SkyBentNormal, DenoiserContext);
+									}
 									DenoiseRawData(Lightmap.GetSize(), IncidentLighting, LuminanceSH, DenoiserContext);
 								}
 							}
@@ -2673,15 +2772,16 @@ void FScene::ApplyFinishedLightmapsToWorld()
 
 							if (bHasSkyShadowing)
 							{
-								CopyRectTiled(
-									FIntPoint(0, 0),
-									FIntRect(FIntPoint(0, 0), Lightmap.GetSize()),
-									SrcRowPitchInPixels,
-									DstRowPitchInPixels,
-									[&Lightmap, &LightSampleData](int32 DstLinearIndex, FIntPoint SrcTilePosition, int32 SrcLinearIndex) mutable
+								for (int32 Y = 0 ; Y < Lightmap.GetSize().Y; Y++)
 								{
-									AddSkyOcclusionToLightSample(LightSampleData[DstLinearIndex], Lightmap.TileStorage[FTileVirtualCoordinates(SrcTilePosition, 0)].CPUTextureData[3]->Data[SrcLinearIndex]);
-								});
+									for (int32 X = 0 ; X < Lightmap.GetSize().X; X++)
+									{
+										int32 LinearIndex = Y * Lightmap.GetSize().X + X;
+										LightSampleData[LinearIndex].SkyOcclusion[0] = SkyBentNormal[LinearIndex].X;
+										LightSampleData[LinearIndex].SkyOcclusion[1] = SkyBentNormal[LinearIndex].Y;
+										LightSampleData[LinearIndex].SkyOcclusion[2] = SkyBentNormal[LinearIndex].Z;
+									}
+								}
 							}
 						}
 
