GPULM: Improve skylight bent normal denoising by applying the same tone mapping trick to work around extreme dynamic ranges

diff --git a/Engine/Plugins/Experimental/GPULightmass/Source/GPULightmass/Private/LightmapDenoising.cpp b/Engine/Plugins/Experimental/GPULightmass/Source/GPULightmass/Private/LightmapDenoising.cpp
--- a/Engine/Plugins/Experimental/GPULightmass/Source/GPULightmass/Private/LightmapDenoising.cpp
+++ b/Engine/Plugins/Experimental/GPULightmass/Source/GPULightmass/Private/LightmapDenoising.cpp
@@ -398,28 +398,33 @@ void DenoiseSkyBentNormal(
 			}
 		}
 	}
-	
-	int32 Padding = 8;
-	
 	// Resizing the filter is a super expensive operation
 	// Round things into size bins to reduce number of resizes
-	FDenoiserFilterSet& FilterSet = DenoiserContext.GetFilterForSize(FIntPoint(FMath::DivideAndRoundUp(Size.X + Padding * 2, 64) * 64, FMath::DivideAndRoundUp(Size.Y + Padding * 2, 64) * 64), true);
-	
+	FDenoiserFilterSet& FilterSet = DenoiserContext.GetFilterForSize(FIntPoint(FMath::DivideAndRoundUp(Size.X, 64) * 64, FMath::DivideAndRoundUp(Size.Y, 64) * 64), true);
+
 	for (int32 Y = 0; Y < Size.Y; Y++)
 	{
 		for (int32 X = 0; X < Size.X; X++)
 		{
-			FilterSet.InputBuffer[(Y + Padding) * FilterSet.Size.X + (X + Padding)] = BentNormal[Y * Size.X + X];
+			const int32 LinearIndex = Y * Size.X + X;
+			const FVector3f N = BentNormal[LinearIndex];
+			FilterSet.InputBuffer[LinearIndex].X = FMath::Sign(N.X) * FMath::Sqrt(FMath::Abs(N.X));
+			FilterSet.InputBuffer[LinearIndex].Y = FMath::Sign(N.Y) * FMath::Sqrt(FMath::Abs(N.Y));
+			FilterSet.InputBuffer[LinearIndex].Z = FMath::Sign(N.Z) * FMath::Sqrt(FMath::Abs(N.Z));			
 		}
 	}
-
+	
 	FilterSet.Execute();
-
+	
 	for (int32 Y = 0; Y < Size.Y; Y++)
 	{
 		for (int32 X = 0; X < Size.X; X++)
 		{
-			BentNormal[Y * Size.X + X] = FilterSet.OutputBuffer[(Y + Padding) * FilterSet.Size.X + (X + Padding)];
+			const int32 LinearIndex = Y * Size.X + X;
+			const FVector3f N = FilterSet.OutputBuffer[LinearIndex];
+			BentNormal[LinearIndex].X = FMath::Sign(N.X) * FMath::Square(N.X);
+			BentNormal[LinearIndex].Y = FMath::Sign(N.Y) * FMath::Square(N.Y);
+			BentNormal[LinearIndex].Z = FMath::Sign(N.Z) * FMath::Square(N.Z);
 		}
 	}
 }
