GPULM: Make InstancedStaticMesh in GPULM use static mesh VF instead of FISMVF. This fixes check(NumUniformBuffersToSet == NumUniformBufferParameters);

diff --git a/Engine/Plugins/Experimental/GPULightmass/Source/GPULightmass/Private/Scene/InstancedStaticMesh.cpp b/Engine/Plugins/Experimental/GPULightmass/Source/GPULightmass/Private/Scene/InstancedStaticMesh.cpp
--- a/Engine/Plugins/Experimental/GPULightmass/Source/GPULightmass/Private/Scene/InstancedStaticMesh.cpp
+++ b/Engine/Plugins/Experimental/GPULightmass/Source/GPULightmass/Private/Scene/InstancedStaticMesh.cpp
@@ -143,20 +143,23 @@ TArray<FMeshBatch> FInstanceGroupRenderState::GetMeshBatchesForGBufferRendering(
 		FMeshBatch MeshBatch;
 
 		FMeshBatchElement& MeshBatchElement = MeshBatch.Elements[0];
-		if (false)// (LODs[LODIndex].OverrideVertexColors)
-		{
-			//MeshBatch.VertexFactory = &StaticMesh->StaticMeshUObject->RenderData->LODVertexFactories[LODIndex].VertexFactoryOverrideColorVertexBuffer;
 
-			//OutMeshBatchElement.VertexFactoryUserData = ProxyLODInfo.OverrideColorVFUniformBuffer.GetReference();
-			//OutMeshBatchElement.UserData = ProxyLODInfo.OverrideColorVertexBuffer;
-			//OutMeshBatchElement.bUserDataIsColorVertexBuffer = true;
+		if (CoordsForCulling.MipLevel == -1)
+		{
+			// No culling, should be for ray tracing scene
+			// YW: While it seems more 'correct' to use FISMVF, what it contains is basically view-dependent instance fading loose parameters,
+			// which isn't particularly useful for lightmap baking.
+			// Most persistent data, like per instance transforms and randoms, has been moved to GPUScene.
+			// This is also consistent with the behavior of FInstancedStaticMeshSceneProxy::GetDynamicRayTracingInstances() which
+			// uses the static mesh VF as well, so that GPULM can benefit from bug fixes there
+			MeshBatch.VertexFactory = &RenderData->LODVertexFactories[LODIndex].VertexFactory;
 		}
 		else
 		{
 			MeshBatch.VertexFactory = &InstancedRenderData->VertexFactories[LODIndex];
-			MeshBatchElement.VertexFactoryUserData = RenderData->LODVertexFactories[LODIndex].VertexFactory.GetUniformBuffer();
 		}
-
+		
+		MeshBatchElement.VertexFactoryUserData = RenderData->LODVertexFactories[LODIndex].VertexFactory.GetUniformBuffer();
 		MeshBatchElement.IndexBuffer = &LODModel.IndexBuffer;
 		MeshBatchElement.FirstIndex = Section.FirstIndex;
 		MeshBatchElement.NumPrimitives = Section.NumTriangles;
