GPULM: Fix use-after-free with dynamic debug names

diff --git a/Engine/Plugins/Experimental/GPULightmass/Source/GPULightmass/Private/LightmapRenderer.cpp b/Engine/Plugins/Experimental/GPULightmass/Source/GPULightmass/Private/LightmapRenderer.cpp
--- a/Engine/Plugins/Experimental/GPULightmass/Source/GPULightmass/Private/LightmapRenderer.cpp
+++ b/Engine/Plugins/Experimental/GPULightmass/Source/GPULightmass/Private/LightmapRenderer.cpp
@@ -1701,7 +1701,9 @@ void FLightmapRenderer::Finalize(FRDGBuilder& GraphBuilder)
 					}
 
 					{
-						FRDGTextureRef RenderTargetTileAtlas = GraphBuilder.RegisterExternalTexture(OutputRenderTargets[RenderTargetIndex], *(FString(TEXT("GPULightmassRenderTargetTileAtlas_")) + StagingTextureNames[RenderTargetIndex]));
+						FString& DynamicDebugName = *GraphBuilder.AllocObject<FString>();
+						DynamicDebugName = TEXT("GPULightmassRenderTargetTileAtlas_") + StagingTextureNames[RenderTargetIndex];
+						FRDGTextureRef RenderTargetTileAtlas = GraphBuilder.RegisterExternalTexture(OutputRenderTargets[RenderTargetIndex], *DynamicDebugName);
 
 						FUploadConvergedLightmapTilesCS::FParameters* PassParameters = GraphBuilder.AllocParameters<FUploadConvergedLightmapTilesCS::FParameters>();
 
diff --git a/Engine/Plugins/Experimental/GPULightmass/Source/GPULightmass/Private/LightmapTilePool.cpp b/Engine/Plugins/Experimental/GPULightmass/Source/GPULightmass/Private/LightmapTilePool.cpp
--- a/Engine/Plugins/Experimental/GPULightmass/Source/GPULightmass/Private/LightmapTilePool.cpp
+++ b/Engine/Plugins/Experimental/GPULightmass/Source/GPULightmass/Private/LightmapTilePool.cpp
@@ -116,6 +116,8 @@ FLightmapTilePoolGPU::FLightmapTilePoolGPU(int32 NumLayers, FIntPoint InSizeInTi
 
 	int64 MemorySize = 0;
 
+	PooledRenderTargetDebugNames.AddDefaulted(NumLayers);
+
 	for (int32 RenderTargetIndex = 0; RenderTargetIndex < NumLayers; RenderTargetIndex++)
 	{
 		LayerFormatAndTileSize.Add( { RenderTargetFormat, TileSize } );
@@ -130,7 +132,8 @@ FLightmapTilePoolGPU::FLightmapTilePoolGPU(int32 NumLayers, FIntPoint InSizeInTi
 
 		FRHICommandListImmediate& RHICmdList = FRHICommandListExecutor::GetImmediateCommandList();
 
-		GRenderTargetPool.FindFreeElement(RHICmdList, Desc, PooledRenderTargets[RenderTargetIndex], *FString::Printf(TEXT("LightmapTilePoolGPU%d"), RenderTargetIndex));
+		PooledRenderTargetDebugNames[RenderTargetIndex] = FString::Printf(TEXT("LightmapTilePoolGPU%d"), RenderTargetIndex);
+		GRenderTargetPool.FindFreeElement(RHICmdList, Desc, PooledRenderTargets[RenderTargetIndex], *PooledRenderTargetDebugNames[RenderTargetIndex]);
 
 		MemorySize += PooledRenderTargets[RenderTargetIndex]->ComputeMemorySize();
 
@@ -148,6 +151,8 @@ void FLightmapTilePoolGPU::Initialize(TArray<FLayerFormatAndTileSize> Layers)
 
 	int64 MemorySize = 0;
 
+	PooledRenderTargetDebugNames.AddDefaulted(Layers.Num());
+	
 	for (int32 RenderTargetIndex = 0; RenderTargetIndex < Layers.Num(); RenderTargetIndex++)
 	{
 		LayerFormatAndTileSize.Add(Layers[RenderTargetIndex]);
@@ -164,7 +169,8 @@ void FLightmapTilePoolGPU::Initialize(TArray<FLayerFormatAndTileSize> Layers)
 
 		FRHICommandListImmediate& RHICmdList = FRHICommandListExecutor::GetImmediateCommandList();
 
-		GRenderTargetPool.FindFreeElement(RHICmdList, Desc, PooledRenderTargets[RenderTargetIndex], *FString::Printf(TEXT("LightmapTilePoolGPU%d"), RenderTargetIndex));
+		PooledRenderTargetDebugNames[RenderTargetIndex] = FString::Printf(TEXT("LightmapTilePoolGPU%d"), RenderTargetIndex);
+		GRenderTargetPool.FindFreeElement(RHICmdList, Desc, PooledRenderTargets[RenderTargetIndex], *PooledRenderTargetDebugNames[RenderTargetIndex]);
 
 		MemorySize += PooledRenderTargets[RenderTargetIndex]->ComputeMemorySize();
 
diff --git a/Engine/Plugins/Experimental/GPULightmass/Source/GPULightmass/Private/LightmapTilePool.h b/Engine/Plugins/Experimental/GPULightmass/Source/GPULightmass/Private/LightmapTilePool.h
--- a/Engine/Plugins/Experimental/GPULightmass/Source/GPULightmass/Private/LightmapTilePool.h
+++ b/Engine/Plugins/Experimental/GPULightmass/Source/GPULightmass/Private/LightmapTilePool.h
@@ -64,6 +64,7 @@ public:
 
 	TArray<FLayerFormatAndTileSize> LayerFormatAndTileSize;
 	TArray<TRefCountPtr<IPooledRenderTarget>> PooledRenderTargets;
+	TArray<FString> PooledRenderTargetDebugNames;
 
 	FIntPoint GetPositionFromLinearAddress(uint32 Address);
 
