GPULM: Leverage on-demand shader compilation to drastically reduce shader compilation triggered by GPULM launch. We no longer reset bake progress on shader compilation because with ODSC it is constantly triggered and resets progress unexpectedly. Also, lightmap GBuffer and VLM voxelization now use default material - they were not considering material features like normal maps and masking anyway.

diff --git a/Engine/Plugins/Experimental/GPULightmass/Shaders/Private/LightmapGBuffer.usf b/Engine/Plugins/Experimental/GPULightmass/Shaders/Private/LightmapGBuffer.usf
--- a/Engine/Plugins/Experimental/GPULightmass/Shaders/Private/LightmapGBuffer.usf
+++ b/Engine/Plugins/Experimental/GPULightmass/Shaders/Private/LightmapGBuffer.usf
@@ -80,18 +80,12 @@ void LightmapGBufferPS(
 	CalcMaterialParametersEx(MaterialParameters, PixelMaterialInputs, float4(0, 0, 0, 0), float4(0, 0, 0, 0), bIsFrontFace, Input.SavedWorldPosition, Input.SavedWorldPosition);
 	
 	float4 OutWorldPosition = float4(LWCHackToFloat(MaterialParameters.AbsoluteWorldPosition) * WorldPositionScalar, 1.0f);
+	// Faceted triangle normal
     float4 OutWorldNormal = float4(normalize(cross(ddx(OutWorldPosition.xyz), ddy(OutWorldPosition.xyz))), 1.0f);
     OutWorldNormal *= CondMask(bIsFrontFace, 1.0f, -1.0f);
 	OutWorldNormal *= GetPrimitive_DeterminantSign(MaterialParameters.PrimitiveId);
-#if 0 // Normal with shader effects (e.g. normal maps). Can bring high frequency artifacts into lightmaps
-	float4 OutShadingNormal = float4(normalize(MaterialParameters.WorldNormal), 0.0f);
-	#if MATERIAL_TANGENTSPACENORMAL
-		// Revert the normal flipping done by CalcPixelMaterialInputs
-		OutShadingNormal *= MaterialParameters.TwoSidedSign;
-	#endif
-#else // Interpolated vertex normal
-	float4 OutShadingNormal = float4(TransformTangentNormalToWorld(MaterialParameters.TangentToWorld, float3(0, 0, 1)), 0.0f);
-#endif
+	// Interpolated vertex normal
+	float4 OutShadingNormal = float4(MaterialParameters.TangentToWorld[2], 0);
 
 	// GBuffer Outputs
 	LightmapGBufferParams.ScratchTilePoolLayer0[ScratchTilePoolOffset + Input.Position.xy] = OutWorldPosition;
diff --git a/Engine/Plugins/Experimental/GPULightmass/Shaders/Private/LightmapPathTracing.usf b/Engine/Plugins/Experimental/GPULightmass/Shaders/Private/LightmapPathTracing.usf
--- a/Engine/Plugins/Experimental/GPULightmass/Shaders/Private/LightmapPathTracing.usf
+++ b/Engine/Plugins/Experimental/GPULightmass/Shaders/Private/LightmapPathTracing.usf
@@ -918,7 +918,7 @@ void LightmapPathTracingMainRG()
 					PathTracingKernel(
 						BatchedTiles[TileIndex].RenderPassIndex,
 						TranslatedWorldPosition,
-						WorldNormal,
+						ShadingNormal,
 						RandSequence,
 						SampleIndex,
 						bIsValidSample,
diff --git a/Engine/Plugins/Experimental/GPULightmass/Source/GPULightmass/Private/GPULightmass.cpp b/Engine/Plugins/Experimental/GPULightmass/Source/GPULightmass/Private/GPULightmass.cpp
--- a/Engine/Plugins/Experimental/GPULightmass/Source/GPULightmass/Private/GPULightmass.cpp
+++ b/Engine/Plugins/Experimental/GPULightmass/Source/GPULightmass/Private/GPULightmass.cpp
@@ -293,11 +293,14 @@ void FGPULightmass::OnLightmassImportanceVolumeModified()
 
 void FGPULightmass::OnMaterialInvalidated(FMaterialRenderProxy* Material)
 {
-	if (Scene.RenderState.CachedRayTracingScene.IsValid())
+	ENQUEUE_RENDER_COMMAND(OnMaterialInvalidatedRenderThread)([&RenderState = Scene.RenderState](FRHICommandListImmediate&) mutable
 	{
-		Scene.RenderState.CachedRayTracingScene.Reset();
-		UE_LOG(LogGPULightmass, Log, TEXT("Cached ray tracing scene is invalidated due to material changes"));
-	}
+		if (RenderState.CachedRayTracingScene.IsValid())
+		{
+			RenderState.CachedRayTracingScene.Reset();
+			UE_LOG(LogGPULightmass, Log, TEXT("Cached ray tracing scene is invalidated due to material changes"));
+		}
+	});
 }
 
 void FGPULightmass::StartRecordingVisibleTiles() 
diff --git a/Engine/Plugins/Experimental/GPULightmass/Source/GPULightmass/Private/GPULightmassSettings.cpp b/Engine/Plugins/Experimental/GPULightmass/Source/GPULightmass/Private/GPULightmassSettings.cpp
--- a/Engine/Plugins/Experimental/GPULightmass/Source/GPULightmass/Private/GPULightmassSettings.cpp
+++ b/Engine/Plugins/Experimental/GPULightmass/Source/GPULightmass/Private/GPULightmassSettings.cpp
@@ -152,9 +152,6 @@ void UGPULightmassSubsystem::Launch()
 		SlowTask.MakeDialog();
 		SlowTask.EnterProgressFrame(1, LOCTEXT("StartingStaticLightingSystem", "Starting static lighting system"));
 
-		UMaterialInterface::SubmitRemainingJobsForWorld(World);
-		FAssetCompilingManager::Get().FinishAllCompilation();
-
 		{
 			FGlobalComponentRecreateRenderStateContext RecreateRenderStateContext; // Implicit FlushRenderingCommands();
 
diff --git a/Engine/Plugins/Experimental/GPULightmass/Source/GPULightmass/Private/LightmapGBuffer.h b/Engine/Plugins/Experimental/GPULightmass/Source/GPULightmass/Private/LightmapGBuffer.h
--- a/Engine/Plugins/Experimental/GPULightmass/Source/GPULightmass/Private/LightmapGBuffer.h
+++ b/Engine/Plugins/Experimental/GPULightmass/Source/GPULightmass/Private/LightmapGBuffer.h
@@ -63,7 +63,8 @@ protected:
 		if (EnumHasAllFlags(Parameters.Flags, EShaderPermutationFlags::HasEditorOnlyData)
 			&& IsFeatureLevelSupported(Parameters.Platform, ERHIFeatureLevel::SM5)
 			&& bAllowStaticLighting
-			&& Parameters.VertexFactoryType->SupportsStaticLighting())
+			&& Parameters.VertexFactoryType->SupportsStaticLighting()
+			&& Parameters.MaterialParameters.bIsSpecialEngineMaterial)
 		{
 			return true;
 		}
@@ -117,7 +118,8 @@ public:
 		if (EnumHasAllFlags(Parameters.Flags, EShaderPermutationFlags::HasEditorOnlyData)
 			&& IsFeatureLevelSupported(Parameters.Platform, ERHIFeatureLevel::SM5)
 			&& bAllowStaticLighting
-			&& Parameters.VertexFactoryType->SupportsStaticLighting())
+			&& Parameters.VertexFactoryType->SupportsStaticLighting()
+			&& Parameters.MaterialParameters.bIsSpecialEngineMaterial)
 		{
 			return true;
 		}
@@ -181,19 +183,17 @@ public:
 
 	virtual void AddMeshBatch(const FMeshBatch& RESTRICT MeshBatch, uint64 BatchElementMask, const FPrimitiveSceneProxy* RESTRICT PrimitiveSceneProxy, int32 StaticMeshId = -1) override final
 	{
-		const FMaterialRenderProxy* FallbackMaterialRenderProxyPtr = nullptr;
-		const FMaterial& Material = MeshBatch.MaterialRenderProxy->GetMaterialWithFallback(FeatureLevel, FallbackMaterialRenderProxyPtr);
-		const FMaterialRenderProxy& MaterialRenderProxy = FallbackMaterialRenderProxyPtr ? *FallbackMaterialRenderProxyPtr : *MeshBatch.MaterialRenderProxy;
-
-		if (MeshBatch.bUseForMaterial
-			&& (!PrimitiveSceneProxy || PrimitiveSceneProxy->ShouldRenderInMainPass()))
+		if (MeshBatch.bUseForMaterial)
 		{
-			Process(MeshBatch, BatchElementMask, StaticMeshId, PrimitiveSceneProxy, MaterialRenderProxy, Material);
+			const FMaterialRenderProxy& DefaultProxy = *UMaterial::GetDefaultMaterial(MD_Surface)->GetRenderProxy();
+			const FMaterial& DefaultMaterial = *DefaultProxy.GetMaterialNoFallback(FeatureLevel);
+
+			Process(MeshBatch, BatchElementMask, StaticMeshId, PrimitiveSceneProxy, DefaultProxy, DefaultMaterial);
 		}
 	}
 
 private:
-	void Process(
+	bool Process(
 		const FMeshBatch& MeshBatch,
 		uint64 BatchElementMask,
 		int32 StaticMeshId,
@@ -203,12 +203,19 @@ private:
 	{
 		const FVertexFactory* VertexFactory = MeshBatch.VertexFactory;
 
+		FMaterialShaderTypes ShaderTypes;
+		ShaderTypes.AddShaderType<FLightmapGBufferVS>();
+		ShaderTypes.AddShaderType<FLightmapGBufferPS>();
+
+		FMaterialShaders MaterialShaders;
+		verify(MaterialResource.TryGetShaders(ShaderTypes, MeshBatch.VertexFactory->GetType(), MaterialShaders));
+		
 		TMeshProcessorShaders<
 			FLightmapGBufferVS,
 			FLightmapGBufferPS> Shaders;
 
-		Shaders.VertexShader = MaterialResource.GetShader<FLightmapGBufferVS>(VertexFactory->GetType());
-		Shaders.PixelShader = MaterialResource.GetShader<FLightmapGBufferPS>(VertexFactory->GetType());
+		MaterialShaders.TryGetVertexShader(Shaders.VertexShader);
+		MaterialShaders.TryGetPixelShader(Shaders.PixelShader);
 
 		const FMeshDrawingPolicyOverrideSettings OverrideSettings = ComputeMeshOverrideSettings(MeshBatch);
 		ERasterizerFillMode MeshFillMode = ComputeMeshFillMode(MeshBatch, MaterialResource, OverrideSettings);
@@ -235,6 +242,8 @@ private:
 			SortKey,
 			EMeshPassFeatures::Default,
 			ShaderElementData);
+
+		return true;
 	}
 
 private:
diff --git a/Engine/Plugins/Experimental/GPULightmass/Source/GPULightmass/Private/LightmapRenderer.cpp b/Engine/Plugins/Experimental/GPULightmass/Source/GPULightmass/Private/LightmapRenderer.cpp
--- a/Engine/Plugins/Experimental/GPULightmass/Source/GPULightmass/Private/LightmapRenderer.cpp
+++ b/Engine/Plugins/Experimental/GPULightmass/Source/GPULightmass/Private/LightmapRenderer.cpp
@@ -657,8 +657,7 @@ void FCachedRayTracingSceneData::SetupFromSceneRenderState(FSceneRenderState& Sc
 
 			for (int32 SegmentIndex = 0; SegmentIndex < MeshBatches.Num(); SegmentIndex++)
 			{
-				const FMaterialRenderProxy* FallbackMaterialRenderProxyPtr = nullptr;
-				const FMaterial& Material = MeshBatches[SegmentIndex].MaterialRenderProxy->GetMaterialWithFallback(Scene.FeatureLevel, FallbackMaterialRenderProxyPtr);
+				const FMaterial& Material = MeshBatches[SegmentIndex].MaterialRenderProxy->GetIncompleteMaterialWithFallback(Scene.FeatureLevel);
 
 				bAllSegmentsUnlit &= Material.GetShadingModels().HasOnlyShadingModel(MSM_Unlit) || !MeshBatches[SegmentIndex].CastShadow;
 				bAllSegmentsOpaque &= Material.GetBlendMode() == EBlendMode::BLEND_Opaque;
@@ -710,8 +709,7 @@ void FCachedRayTracingSceneData::SetupFromSceneRenderState(FSceneRenderState& Sc
 
 			for (int32 SegmentIndex = 0; SegmentIndex < MeshBatches.Num(); SegmentIndex++)
 			{
-				const FMaterialRenderProxy* FallbackMaterialRenderProxyPtr = nullptr;
-				const FMaterial& Material = MeshBatches[SegmentIndex].MaterialRenderProxy->GetMaterialWithFallback(Scene.FeatureLevel, FallbackMaterialRenderProxyPtr);
+				const FMaterial& Material = MeshBatches[SegmentIndex].MaterialRenderProxy->GetIncompleteMaterialWithFallback(Scene.FeatureLevel);
 
 				bAllSegmentsUnlit &= Material.GetShadingModels().HasOnlyShadingModel(MSM_Unlit) || !MeshBatches[SegmentIndex].CastShadow;
 				bAllSegmentsOpaque &= Material.GetBlendMode() == EBlendMode::BLEND_Opaque;
@@ -768,6 +766,14 @@ void FCachedRayTracingSceneData::SetupFromSceneRenderState(FSceneRenderState& Sc
 #endif // RHI_RAYTRACING
 }
 
+FCachedRayTracingSceneData::~FCachedRayTracingSceneData()
+{
+	FRHICommandListImmediate& RHICmdList = FRHICommandListExecutor::GetImmediateCommandList();
+
+	// Move OwnedRayTracingInstanceTransforms to RHIThread and it will get destroyed when the lambda is released
+	RHICmdList.EnqueueLambda([OwnedRayTracingInstanceTransforms = MoveTemp(OwnedRayTracingInstanceTransforms)](FRHICommandList&){});
+}
+
 bool FSceneRenderState::SetupRayTracingScene(int32 LODIndex)
 {
 	TRACE_CPUPROFILER_EVENT_SCOPE(SetupRayTracingScene);
@@ -917,6 +923,8 @@ bool FSceneRenderState::SetupRayTracingScene(int32 LODIndex)
 		FRayTracingMeshCommandOneFrameArray VisibleRayTracingMeshCommands;
 		FDynamicRayTracingMeshCommandStorage DynamicRayTracingMeshCommandStorage;
 
+		TArray<TUniquePtr<FMatrix>> LandscapeTransforms;
+		
 		{
 			TRACE_CPUPROFILER_EVENT_SCOPE(Landscapes);
 
@@ -978,6 +986,12 @@ bool FSceneRenderState::SetupRayTracingScene(int32 LODIndex)
 							MeshBatches[0].Elements[0].DynamicPrimitiveIndex += StaticMeshInstanceRenderStates.Elements.Num();
 							MeshBatches[0].Elements[0].DynamicPrimitiveIndex += InstanceGroupRenderStates.Elements.Num();
 
+							for (FMeshBatch& MeshBatch : MeshBatches)
+							{
+								// Override with default material as we're not considering WPO in GPULM landscape creation
+								MeshBatch.MaterialRenderProxy = UMaterial::GetDefaultMaterial(MD_Surface)->GetRenderProxy();
+							}
+							
 							FRayTracingDynamicGeometryUpdateParams UpdateParams
 							{
 								MeshBatches,
@@ -1013,7 +1027,7 @@ bool FSceneRenderState::SetupRayTracingScene(int32 LODIndex)
 
 						FRayTracingGeometryInstance& RayTracingInstance = RayTracingGeometryInstances[InstanceIndex];
 						RayTracingInstance.GeometryRHI = Landscape.SectionRayTracingStates[SubSectionIdx]->Geometry.RayTracingGeometryRHI;
-						RayTracingInstance.Transforms = MakeArrayView(&FMatrix::Identity, 1);
+						RayTracingInstance.Transforms = MakeArrayView(LandscapeTransforms.Emplace_GetRef(MakeUnique<FMatrix>(Landscape.LocalToWorld)).Get(), 1);
 						RayTracingInstance.NumTransforms = 1;
 						RayTracingInstance.DefaultUserData = (uint32)(StaticMeshInstanceRenderStates.Elements.Num() + InstanceGroupRenderStates.Elements.Num() + LandscapeIndex);
 
@@ -1041,8 +1055,7 @@ bool FSceneRenderState::SetupRayTracingScene(int32 LODIndex)
 
 							RayTracingMeshProcessor.AddMeshBatch(MeshBatches[SegmentIndex], 1, nullptr);
 
-							const FMaterialRenderProxy* FallbackMaterialRenderProxyPtr = nullptr;
-							const FMaterial& Material = MeshBatches[SegmentIndex].MaterialRenderProxy->GetMaterialWithFallback(FeatureLevel, FallbackMaterialRenderProxyPtr);
+							const FMaterial& Material = MeshBatches[SegmentIndex].MaterialRenderProxy->GetIncompleteMaterialWithFallback(FeatureLevel);
 
 							bAllSegmentsUnlit &= Material.GetShadingModels().HasOnlyShadingModel(MSM_Unlit) || !MeshBatches[SegmentIndex].CastShadow;
 							bAllSegmentsOpaque &= Material.GetBlendMode() == EBlendMode::BLEND_Opaque;
@@ -1191,6 +1204,9 @@ bool FSceneRenderState::SetupRayTracingScene(int32 LODIndex)
 				RHICmdList.BuildAccelerationStructure(BuildParams);
 			}
 
+			// Move LandscapeTransforms to RHIThread to extend its lifetime until RHI cmd execution
+			RHICmdList.EnqueueLambda([LandscapeTransforms = MoveTemp(LandscapeTransforms)](FRHICommandList&){});
+			
 			FRayTracingPipelineStateInitializer PSOInitializer;
 
 			PSOInitializer.MaxPayloadSizeInBytes = RAY_TRACING_MAX_ALLOWED_PAYLOAD_SIZE;
@@ -1947,7 +1963,10 @@ void FLightmapRenderer::Finalize(FRDGBuilder& GraphBuilder)
 
 	RectLightAtlas::UpdateRectLightAtlasTexture(GraphBuilder, Scene->FeatureLevel);
 
-	Scene->SetupRayTracingScene(MostCommonLODIndex);
+	if (!Scene->SetupRayTracingScene(MostCommonLODIndex))
+	{
+		return;
+	}
 
 	TStaticArray<FRDGTextureUAVRef, 3> ScratchTilePoolLayerUAVs;
 
diff --git a/Engine/Plugins/Experimental/GPULightmass/Source/GPULightmass/Private/Scene/Scene.cpp b/Engine/Plugins/Experimental/GPULightmass/Source/GPULightmass/Private/Scene/Scene.cpp
--- a/Engine/Plugins/Experimental/GPULightmass/Source/GPULightmass/Private/Scene/Scene.cpp
+++ b/Engine/Plugins/Experimental/GPULightmass/Source/GPULightmass/Private/Scene/Scene.cpp
@@ -369,7 +369,6 @@ void FScene::AddLight(LightComponentType* LightComponent)
 {
 	if (LightTypeInfo<LightComponentType>::GetLightComponentRegistration(LightScene).Contains(LightComponent))
 	{
-		UE_LOG(LogGPULightmass, Log, TEXT("Warning: duplicated component registration"));
 		return;
 	}
 
@@ -576,7 +575,6 @@ void FScene::AddLight(USkyLightComponent* SkyLight)
 {
 	if (LightScene.SkyLight.IsSet() && LightScene.SkyLight->ComponentUObject == SkyLight)
 	{
-		UE_LOG(LogGPULightmass, Log, TEXT("Warning: duplicated component registration"));
 		return;
 	}
 
@@ -684,7 +682,6 @@ void FScene::AddGeometryInstanceFromComponent(UStaticMeshComponent* InComponent)
 {
 	if (RegisteredStaticMeshComponentUObjects.Contains(InComponent))
 	{
-		UE_LOG(LogGPULightmass, Log, TEXT("Warning: duplicated component registration"));
 		return;
 	}
 	
@@ -896,13 +893,11 @@ void FScene::AddGeometryInstanceFromComponent(UInstancedStaticMeshComponent* InC
 {
 	if (InComponent->PerInstanceSMData.Num() == 0)
 	{
-		UE_LOG(LogGPULightmass, Log, TEXT("Skipping empty instanced static mesh"));
 		return;
 	}
 
 	if (RegisteredInstancedStaticMeshComponentUObjects.Contains(InComponent))
 	{
-		UE_LOG(LogGPULightmass, Log, TEXT("Warning: duplicated component registration"));
 		return;
 	}
 
@@ -1112,13 +1107,11 @@ void FScene::AddGeometryInstanceFromComponent(ULandscapeComponent* InComponent)
 {
 	if (InComponent->GetLandscapeInfo() == nullptr)
 	{
-		UE_LOG(LogGPULightmass, Log, TEXT("Skipping landscape with empty info object"));
 		return;
 	}
 
 	if (RegisteredLandscapeComponentUObjects.Contains(InComponent))
 	{
-		UE_LOG(LogGPULightmass, Log, TEXT("Warning: duplicated component registration"));
 		return;
 	}
 
diff --git a/Engine/Plugins/Experimental/GPULightmass/Source/GPULightmass/Private/Scene/Scene.h b/Engine/Plugins/Experimental/GPULightmass/Source/GPULightmass/Private/Scene/Scene.h
--- a/Engine/Plugins/Experimental/GPULightmass/Source/GPULightmass/Private/Scene/Scene.h
+++ b/Engine/Plugins/Experimental/GPULightmass/Source/GPULightmass/Private/Scene/Scene.h
@@ -58,6 +58,8 @@ private:
 
 struct FCachedRayTracingSceneData
 {
+	~FCachedRayTracingSceneData();
+	
 	TArray<TArray<FVisibleRayTracingMeshCommand>> VisibleRayTracingMeshCommandsPerLOD;
 	TChunkedArray<FRayTracingMeshCommand> MeshCommandStorage;
 
diff --git a/Engine/Plugins/Experimental/GPULightmass/Source/GPULightmass/Private/VolumetricLightmap.h b/Engine/Plugins/Experimental/GPULightmass/Source/GPULightmass/Private/VolumetricLightmap.h
--- a/Engine/Plugins/Experimental/GPULightmass/Source/GPULightmass/Private/VolumetricLightmap.h
+++ b/Engine/Plugins/Experimental/GPULightmass/Source/GPULightmass/Private/VolumetricLightmap.h
@@ -24,7 +24,7 @@ public:
 	FBox CombinedImportanceVolume;
 	TArray<FBox> ImportanceVolumes;
 	float TargetDetailCellSize = 50.0f;
-	int32 NumTotalBricks;
+	int32 NumTotalBricks = 0;
 
 	int32 FrameNumber = 0;
 	int32 NumTotalPassesToRender = 0;
diff --git a/Engine/Plugins/Experimental/GPULightmass/Source/GPULightmass/Private/VolumetricLightmapVoxelization.h b/Engine/Plugins/Experimental/GPULightmass/Source/GPULightmass/Private/VolumetricLightmapVoxelization.h
--- a/Engine/Plugins/Experimental/GPULightmass/Source/GPULightmass/Private/VolumetricLightmapVoxelization.h
+++ b/Engine/Plugins/Experimental/GPULightmass/Source/GPULightmass/Private/VolumetricLightmapVoxelization.h
@@ -45,7 +45,8 @@ protected:
 	{
 		return EnumHasAllFlags(Parameters.Flags, EShaderPermutationFlags::HasEditorOnlyData)
 			&& IsFeatureLevelSupported(Parameters.Platform, ERHIFeatureLevel::SM5)
-			&& IsSupportedVertexFactoryType(Parameters.VertexFactoryType);
+			&& IsSupportedVertexFactoryType(Parameters.VertexFactoryType)
+			&& Parameters.MaterialParameters.bIsSpecialEngineMaterial;
 	}
 };
 
@@ -68,7 +69,8 @@ protected:
 	{
 		return EnumHasAllFlags(Parameters.Flags, EShaderPermutationFlags::HasEditorOnlyData)
 			&& IsFeatureLevelSupported(Parameters.Platform, ERHIFeatureLevel::SM5)
-			&& IsSupportedVertexFactoryType(Parameters.VertexFactoryType);
+			&& IsSupportedVertexFactoryType(Parameters.VertexFactoryType)
+			&& Parameters.MaterialParameters.bIsSpecialEngineMaterial;
 	}
 };
 
@@ -91,7 +93,8 @@ protected:
 	{
 		return EnumHasAllFlags(Parameters.Flags, EShaderPermutationFlags::HasEditorOnlyData)
 			&& IsFeatureLevelSupported(Parameters.Platform, ERHIFeatureLevel::SM5)
-			&& IsSupportedVertexFactoryType(Parameters.VertexFactoryType);
+			&& IsSupportedVertexFactoryType(Parameters.VertexFactoryType)
+			&& Parameters.MaterialParameters.bIsSpecialEngineMaterial;
 	}
 };
 
@@ -107,21 +110,22 @@ public:
 
 	virtual void AddMeshBatch(const FMeshBatch& RESTRICT MeshBatch, uint64 BatchElementMask, const FPrimitiveSceneProxy* RESTRICT PrimitiveSceneProxy, int32 StaticMeshId = -1) override final
 	{
-		const FMaterialRenderProxy* FallbackMaterialRenderProxyPtr = nullptr;
-		const FMaterial& Material = MeshBatch.MaterialRenderProxy->GetMaterialWithFallback(FeatureLevel, FallbackMaterialRenderProxyPtr);
-		const FMaterialRenderProxy& MaterialRenderProxy = FallbackMaterialRenderProxyPtr ? *FallbackMaterialRenderProxyPtr : *MeshBatch.MaterialRenderProxy;
-
-		if (MeshBatch.bUseForMaterial
-			&& (Material.GetBlendMode() == BLEND_Opaque || Material.IsMasked())
-			&& (!PrimitiveSceneProxy || PrimitiveSceneProxy->ShouldRenderInMainPass())
-			&& IsSupportedVertexFactoryType(MeshBatch.VertexFactory->GetType()))
+		if (MeshBatch.bUseForMaterial && IsSupportedVertexFactoryType(MeshBatch.VertexFactory->GetType()))
 		{
-			Process(MeshBatch, BatchElementMask, StaticMeshId, PrimitiveSceneProxy, MaterialRenderProxy, Material);
+			const FMaterial& Material = MeshBatch.MaterialRenderProxy->GetIncompleteMaterialWithFallback(FeatureLevel);
+
+			if (Material.GetBlendMode() == BLEND_Opaque || Material.IsMasked())
+			{
+				const FMaterialRenderProxy& DefaultProxy = *UMaterial::GetDefaultMaterial(MD_Surface)->GetRenderProxy();
+				const FMaterial& DefaultMaterial = *DefaultProxy.GetMaterialNoFallback(FeatureLevel);
+
+				Process(MeshBatch, BatchElementMask, StaticMeshId, PrimitiveSceneProxy, DefaultProxy, DefaultMaterial);
+			}
 		}
 	}
 
 private:
-	void Process(
+	bool Process(
 		const FMeshBatch& MeshBatch,
 		uint64 BatchElementMask,
 		int32 StaticMeshId,
@@ -131,14 +135,23 @@ private:
 	{
 		const FVertexFactory* VertexFactory = MeshBatch.VertexFactory;
 
+		FMaterialShaderTypes ShaderTypes;
+		ShaderTypes.AddShaderType<FVLMVoxelizationVS>();
+		ShaderTypes.AddShaderType<FVLMVoxelizationGS>();
+		ShaderTypes.AddShaderType<FVLMVoxelizationPS>();
+
+		FMaterialShaders MaterialShaders;
+		// We're using DefaultMaterial and it should never fail
+		verify(MaterialResource.TryGetShaders(ShaderTypes, MeshBatch.VertexFactory->GetType(), MaterialShaders));
+		
 		TMeshProcessorShaders<
 			FVLMVoxelizationVS,
 			FVLMVoxelizationPS,
 			FVLMVoxelizationGS> Shaders;
 
-		Shaders.VertexShader = MaterialResource.GetShader<FVLMVoxelizationVS>(VertexFactory->GetType());
-		Shaders.PixelShader = MaterialResource.GetShader<FVLMVoxelizationPS>(VertexFactory->GetType());
-		Shaders.GeometryShader = MaterialResource.GetShader<FVLMVoxelizationGS>(VertexFactory->GetType());
+		MaterialShaders.TryGetVertexShader(Shaders.VertexShader);
+		MaterialShaders.TryGetGeometryShader(Shaders.GeometryShader);
+		MaterialShaders.TryGetPixelShader(Shaders.PixelShader);
 
 		const FMeshDrawingPolicyOverrideSettings OverrideSettings = ComputeMeshOverrideSettings(MeshBatch);
 		ERasterizerFillMode MeshFillMode = ComputeMeshFillMode(MeshBatch, MaterialResource, OverrideSettings);
@@ -162,6 +175,8 @@ private:
 			SortKey,
 			EMeshPassFeatures::Default,
 			ShaderElementData);
+
+		return true;
 	}
 
 private:
diff --git a/Engine/Plugins/Experimental/GPULightmass/Source/GPULightmassEditor/Private/GPULightmassEditorModule.cpp b/Engine/Plugins/Experimental/GPULightmass/Source/GPULightmassEditor/Private/GPULightmassEditorModule.cpp
--- a/Engine/Plugins/Experimental/GPULightmass/Source/GPULightmassEditor/Private/GPULightmassEditorModule.cpp
+++ b/Engine/Plugins/Experimental/GPULightmass/Source/GPULightmassEditor/Private/GPULightmassEditorModule.cpp
@@ -24,8 +24,6 @@
 #include "Widgets/Input/SCheckBox.h"
 #include "ToolMenus.h"
 
-extern ENGINE_API void ToggleLightmapPreview_GameThread(UWorld* InWorld);
-
 #define LOCTEXT_NAMESPACE "StaticLightingSystem"
 
 IMPLEMENT_MODULE( FGPULightmassEditorModule, GPULightmassEditor )
diff --git a/Engine/Source/Runtime/Engine/Private/ShaderCompiler/ShaderCompiler.cpp b/Engine/Source/Runtime/Engine/Private/ShaderCompiler/ShaderCompiler.cpp
--- a/Engine/Source/Runtime/Engine/Private/ShaderCompiler/ShaderCompiler.cpp
+++ b/Engine/Source/Runtime/Engine/Private/ShaderCompiler/ShaderCompiler.cpp
@@ -4667,14 +4667,6 @@ void FShaderCompilingManager::PropagateMaterialChangesToPrimitives(TMap<TRefCoun
 				if (bPrimitiveIsDependentOnMaterial)
 				{
 					ComponentContexts.Add(new FComponentRecreateRenderStateContext(PrimitiveComponent, &ScenesToUpdate));
-#if WITH_EDITOR
-					FStaticLightingSystemInterface::OnPrimitiveComponentUnregistered.Broadcast(PrimitiveComponent);
-
-					if (PrimitiveComponent->HasValidSettingsForStaticLighting(false))
-					{
-						FStaticLightingSystemInterface::OnPrimitiveComponentRegistered.Broadcast(PrimitiveComponent);
-					}
-#endif
 				}
 			}
 		}
