GPULM: Implement preliminary support for SkyAtmosphere

diff --git a/Engine/Plugins/Experimental/GPULightmass/Source/GPULightmass/Private/GPULightmass.cpp b/Engine/Plugins/Experimental/GPULightmass/Source/GPULightmass/Private/GPULightmass.cpp
--- a/Engine/Plugins/Experimental/GPULightmass/Source/GPULightmass/Private/GPULightmass.cpp
+++ b/Engine/Plugins/Experimental/GPULightmass/Source/GPULightmass/Private/GPULightmass.cpp
@@ -103,6 +103,7 @@ void FGPULightmass::InstallGameThreadEventHooks()
 	FStaticLightingSystemInterface::OnStationaryLightChannelReassigned.AddRaw(this, &FGPULightmass::OnStationaryLightChannelReassigned);
 	FStaticLightingSystemInterface::OnLightmassImportanceVolumeModified.AddRaw(this, &FGPULightmass::OnLightmassImportanceVolumeModified);
 	FStaticLightingSystemInterface::OnMaterialInvalidated.AddRaw(this, &FGPULightmass::OnMaterialInvalidated);
+	FStaticLightingSystemInterface::OnSkyAtmosphereModified.AddRaw(this, &FGPULightmass::OnSkyAtmosphereModified);
 }
 
 void FGPULightmass::RemoveGameThreadEventHooks()
@@ -116,6 +117,7 @@ void FGPULightmass::RemoveGameThreadEventHooks()
 	FStaticLightingSystemInterface::OnStationaryLightChannelReassigned.RemoveAll(this);
 	FStaticLightingSystemInterface::OnLightmassImportanceVolumeModified.RemoveAll(this);
 	FStaticLightingSystemInterface::OnMaterialInvalidated.RemoveAll(this);
+	FStaticLightingSystemInterface::OnSkyAtmosphereModified.RemoveAll(this);
 }
 
 void FGPULightmass::OnPrimitiveComponentRegistered(UPrimitiveComponent* InComponent)
@@ -168,6 +170,7 @@ void FGPULightmass::OnLightComponentRegistered(ULightComponentBase* InComponent)
 
 	if (UDirectionalLightComponent* DirectionalLight = Cast<UDirectionalLightComponent>(InComponent))
 	{
+		Scene.ConditionalTriggerSkyLightRecapture();
 		Scene.AddLight(DirectionalLight);
 	}
 	else if (URectLightComponent* RectLight = Cast<URectLightComponent>(InComponent))
@@ -195,6 +198,7 @@ void FGPULightmass::OnLightComponentUnregistered(ULightComponentBase* InComponen
 	if (UDirectionalLightComponent* DirectionalLight = Cast<UDirectionalLightComponent>(InComponent))
 	{
 		Scene.RemoveLight(DirectionalLight);
+		Scene.ConditionalTriggerSkyLightRecapture();
 	}
 	else if (URectLightComponent* RectLight = Cast<URectLightComponent>(InComponent))
 	{
@@ -252,6 +256,11 @@ void FGPULightmass::OnStationaryLightChannelReassigned(ULightComponentBase* InCo
 	}
 }
 
+void FGPULightmass::OnSkyAtmosphereModified()
+{
+	Scene.OnSkyAtmosphereModified();
+}
+
 void FGPULightmass::OnPreWorldFinishDestroy(UWorld* InWorld)
 {
 	if (InWorld != World) return;
diff --git a/Engine/Plugins/Experimental/GPULightmass/Source/GPULightmass/Private/GPULightmass.h b/Engine/Plugins/Experimental/GPULightmass/Source/GPULightmass/Private/GPULightmass.h
--- a/Engine/Plugins/Experimental/GPULightmass/Source/GPULightmass/Private/GPULightmass.h
+++ b/Engine/Plugins/Experimental/GPULightmass/Source/GPULightmass/Private/GPULightmass.h
@@ -53,4 +53,5 @@ private:
 	void OnStationaryLightChannelReassigned(ULightComponentBase* InComponent, int32 NewShadowMapChannel);
 	void OnLightmassImportanceVolumeModified();
 	void OnMaterialInvalidated(FMaterialRenderProxy* Material);
+	void OnSkyAtmosphereModified();
 };
diff --git a/Engine/Plugins/Experimental/GPULightmass/Source/GPULightmass/Private/Scene/Lights.cpp b/Engine/Plugins/Experimental/GPULightmass/Source/GPULightmass/Private/Scene/Lights.cpp
--- a/Engine/Plugins/Experimental/GPULightmass/Source/GPULightmass/Private/Scene/Lights.cpp
+++ b/Engine/Plugins/Experimental/GPULightmass/Source/GPULightmass/Private/Scene/Lights.cpp
@@ -4,6 +4,8 @@
 #include "RenderGraphBuilder.h"
 #include "ReflectionEnvironment.h"
 #include "RectLightTextureManager.h"
+#include "Components/SkyAtmosphereComponent.h"
+#include "UObject/UObjectIterator.h"
 
 RENDERER_API void PrepareSkyTexture_Internal(
 	FRDGBuilder& GraphBuilder,
@@ -44,7 +46,7 @@ FLocalLightBuildInfo::FLocalLightBuildInfo(ULightComponent* LightComponent)
 	bStationary = LightComponent->Mobility == EComponentMobility::Stationary;
 	bCastShadow = LightComponent->CastShadows && LightComponent->CastStaticShadows;
 	ShadowMapChannel = LightComponent->PreviewShadowMapChannel;
-	LightComponentMapBuildData = MakeUnique<FLightComponentMapBuildData>();
+	LightComponentMapBuildData = MakeShared<FLightComponentMapBuildData>();
 	LightComponentMapBuildData->ShadowMapChannel = ShadowMapChannel;
 }
 
@@ -65,6 +67,21 @@ FDirectionalLightRenderState::FDirectionalLightRenderState(UDirectionalLightComp
 	: FLocalLightRenderState(DirectionalLightComponent)
 {
 	Color = DirectionalLightComponent->GetColoredLightBrightness();
+	
+	if (DirectionalLightComponent->IsUsedAsAtmosphereSunLight())
+	{
+		for (USkyAtmosphereComponent* SkyAtmosphere : TObjectRange<USkyAtmosphereComponent>(RF_ClassDefaultObject | RF_ArchetypeObject, true, EInternalObjectFlags::Garbage))
+		{
+			if (SkyAtmosphere->GetWorld() == DirectionalLightComponent->GetWorld())
+			{
+				FAtmosphereSetup AtmosphereSetup(*SkyAtmosphere);
+				FLinearColor SunLightAtmosphereTransmittance = AtmosphereSetup.GetTransmittanceAtGroundLevel(-DirectionalLightComponent->GetDirection());
+				Color *= SunLightAtmosphereTransmittance;
+				break; // We only register the first we find
+			}
+		}
+	}
+	
 	Direction = DirectionalLightComponent->GetDirection();
 	LightSourceAngle = DirectionalLightComponent->LightSourceAngle;
 	LightSourceSoftAngle = DirectionalLightComponent->LightSourceSoftAngle;
diff --git a/Engine/Plugins/Experimental/GPULightmass/Source/GPULightmass/Private/Scene/Lights.h b/Engine/Plugins/Experimental/GPULightmass/Source/GPULightmass/Private/Scene/Lights.h
--- a/Engine/Plugins/Experimental/GPULightmass/Source/GPULightmass/Private/Scene/Lights.h
+++ b/Engine/Plugins/Experimental/GPULightmass/Source/GPULightmass/Private/Scene/Lights.h
@@ -73,7 +73,8 @@ struct FLocalLightBuildInfo
 	bool bCastShadow = true;
 	int ShadowMapChannel = INDEX_NONE;
 
-	TUniquePtr<FLightComponentMapBuildData> LightComponentMapBuildData;
+	// This will also be held by FLocalLightRenderState to extend its lifetime until RenderThread has finished with it
+	TSharedPtr<FLightComponentMapBuildData, ESPMode::ThreadSafe> LightComponentMapBuildData;
 
 	virtual bool AffectsBounds(const FBoxSphereBounds& InBounds) const = 0;
 	virtual ULightComponent* GetComponentUObject() const = 0;
@@ -92,6 +93,8 @@ struct FLocalLightRenderState
 	bool bStationary = false;
 	bool bCastShadow = true;
 	int ShadowMapChannel = INDEX_NONE;
+	
+	TSharedPtr<FLightComponentMapBuildData, ESPMode::ThreadSafe> LightComponentMapBuildData;
 
 	virtual FLightRenderParameters GetLightShaderParameters() const = 0;
 };
diff --git a/Engine/Plugins/Experimental/GPULightmass/Source/GPULightmass/Private/Scene/Scene.cpp b/Engine/Plugins/Experimental/GPULightmass/Source/GPULightmass/Private/Scene/Scene.cpp
--- a/Engine/Plugins/Experimental/GPULightmass/Source/GPULightmass/Private/Scene/Scene.cpp
+++ b/Engine/Plugins/Experimental/GPULightmass/Source/GPULightmass/Private/Scene/Scene.cpp
@@ -389,6 +389,7 @@ void FScene::AddLight(LightComponentType* LightComponent)
 	LightTypeInfo<LightComponentType>::GetLightComponentRegistration(LightScene).Add(LightComponent, LightRef);
 
 	typename LightTypeInfo<LightComponentType>::RenderStateType LightRenderState(LightComponent);
+	LightRenderState.LightComponentMapBuildData = LightRef->LightComponentMapBuildData;
 
 	TArray<FPrimitiveSceneProxy*> SceneProxiesToUpdateOnRenderThread;
 	TArray<FGeometryRenderStateToken> RelevantGeometriesToUpdateOnRenderThread;
@@ -648,6 +649,34 @@ void FScene::RemoveLight(USkyLightComponent* SkyLight)
 	});
 }
 
+void FScene::OnSkyAtmosphereModified()
+{
+	TArray<UDirectionalLightComponent*> PreviouslyRegisteredDirectionalLight;
+	LightScene.RegisteredDirectionalLightComponentUObjects.GenerateKeyArray(PreviouslyRegisteredDirectionalLight);	
+	for(UDirectionalLightComponent* DirectionalLight : PreviouslyRegisteredDirectionalLight)
+	{
+		RemoveLight(DirectionalLight);
+		AddLight(DirectionalLight);
+	}
+
+	ConditionalTriggerSkyLightRecapture();
+}
+
+void FScene::ConditionalTriggerSkyLightRecapture()
+{
+	if (LightScene.SkyLight.IsSet())
+	{
+		USkyLightComponent* SkyLight = LightScene.SkyLight->ComponentUObject;
+		if (SkyLight->SourceType == SLS_CapturedScene || SkyLight->bRealTimeCapture)
+		{
+			RemoveLight(SkyLight);
+			SkyLight->SetCaptureIsDirty();
+			SkyLight->UpdateSkyCaptureContents(SkyLight->GetWorld());
+			AddLight(SkyLight);
+		}
+	}
+}
+
 template<typename LightType, typename GeometryRefType>
 TArray<int32> AddAllPossiblyRelevantLightsToGeometry(
 	TEntityArray<LightType>& LightArray,
diff --git a/Engine/Plugins/Experimental/GPULightmass/Source/GPULightmass/Private/Scene/Scene.h b/Engine/Plugins/Experimental/GPULightmass/Source/GPULightmass/Private/Scene/Scene.h
--- a/Engine/Plugins/Experimental/GPULightmass/Source/GPULightmass/Private/Scene/Scene.h
+++ b/Engine/Plugins/Experimental/GPULightmass/Source/GPULightmass/Private/Scene/Scene.h
@@ -176,6 +176,8 @@ public:
 	bool HasLight(LightComponentType* Light);
 
 	void GatherImportanceVolumes();
+	void OnSkyAtmosphereModified();
+	void ConditionalTriggerSkyLightRecapture();
 	
 	void BackgroundTick();
 	void AddRelevantStaticLightGUIDs(FQuantizedLightmapData* QuantizedLightmapData, const FBoxSphereBounds& WorldBounds);
diff --git a/Engine/Source/Editor/UnrealEd/Private/LevelEditorViewport.cpp b/Engine/Source/Editor/UnrealEd/Private/LevelEditorViewport.cpp
--- a/Engine/Source/Editor/UnrealEd/Private/LevelEditorViewport.cpp
+++ b/Engine/Source/Editor/UnrealEd/Private/LevelEditorViewport.cpp
@@ -96,6 +96,7 @@
 #include "UnrealWidget.h"
 #include "EdModeInteractiveToolsContext.h"
 #include "GenericPlatform/GenericPlatformInputDeviceMapper.h"
+#include "Rendering/StaticLightingSystemInterface.h"
 
 DEFINE_LOG_CATEGORY(LogEditorViewport);
 
@@ -187,6 +188,13 @@ static void NotifyAtmosphericLightHasMoved(UDirectionalLightComponent& SelectedA
 		// Now notify the owner about the transform update, e.g. construction script on instance.
 		LightOwner->PostEditMove(bFinished);
 		// No PostEditChangeProperty because not paired with a PreEditChange
+
+		if (bFinished)
+		{
+			SelectedAtmosphericLight.InvalidateLightingCache();
+		}
+		
+		FStaticLightingSystemInterface::OnSkyAtmosphereModified.Broadcast();
 	}
 }
 
diff --git a/Engine/Source/Runtime/Engine/Private/Components/SkyAtmosphereComponent.cpp b/Engine/Source/Runtime/Engine/Private/Components/SkyAtmosphereComponent.cpp
--- a/Engine/Source/Runtime/Engine/Private/Components/SkyAtmosphereComponent.cpp
+++ b/Engine/Source/Runtime/Engine/Private/Components/SkyAtmosphereComponent.cpp
@@ -18,6 +18,7 @@
 
 #if WITH_EDITOR
 #include "ObjectEditorUtils.h"
+#include "Rendering/StaticLightingSystemInterface.h"
 #endif
 
 #define LOCTEXT_NAMESPACE "SkyAtmosphereComponent"
@@ -259,6 +260,11 @@ void USkyAtmosphereComponent::PostEditChangeProperty(FPropertyChangedEvent& Prop
 		{
 			SendRenderTransformCommand();
 		}
+
+#if WITH_EDITOR
+		FStaticLightingSystemInterface::OnSkyAtmosphereModified.Broadcast();
+#endif
+
 	}
 }
 
diff --git a/Engine/Source/Runtime/Engine/Private/Rendering/StaticLightingSystemInterface.cpp b/Engine/Source/Runtime/Engine/Private/Rendering/StaticLightingSystemInterface.cpp
--- a/Engine/Source/Runtime/Engine/Private/Rendering/StaticLightingSystemInterface.cpp
+++ b/Engine/Source/Runtime/Engine/Private/Rendering/StaticLightingSystemInterface.cpp
@@ -17,6 +17,7 @@ FStaticLightingSystemInterface::FLightComponentBasedSignature FStaticLightingSys
 FStaticLightingSystemInterface::FStationaryLightChannelReassignmentSignature FStaticLightingSystemInterface::OnStationaryLightChannelReassigned;
 FStaticLightingSystemInterface::FLightmassImportanceVolumeModifiedSignature FStaticLightingSystemInterface::OnLightmassImportanceVolumeModified;
 FStaticLightingSystemInterface::FMaterialInvalidationSignature FStaticLightingSystemInterface::OnMaterialInvalidated;
+FSimpleMulticastDelegate FStaticLightingSystemInterface::OnSkyAtmosphereModified;
 
 FStaticLightingSystemInterface* FStaticLightingSystemInterface::Get()
 {
diff --git a/Engine/Source/Runtime/Engine/Public/Rendering/StaticLightingSystemInterface.h b/Engine/Source/Runtime/Engine/Public/Rendering/StaticLightingSystemInterface.h
--- a/Engine/Source/Runtime/Engine/Public/Rendering/StaticLightingSystemInterface.h
+++ b/Engine/Source/Runtime/Engine/Public/Rendering/StaticLightingSystemInterface.h
@@ -3,11 +3,10 @@
 #pragma once
 #include "Components/PrimitiveComponent.h"
 #include "Components/LightComponent.h"
+#include "Components/SkyAtmosphereComponent.h"
 
 class FMaterialRenderProxy;
 
-// TODO: Remove the 'V2' suffix after the renderer version is removed
-
 class ENGINE_API IStaticLightingSystem
 {
 public:
@@ -42,9 +41,12 @@ public:
 	static FStationaryLightChannelReassignmentSignature OnStationaryLightChannelReassigned;
 	static FLightmassImportanceVolumeModifiedSignature OnLightmassImportanceVolumeModified;
 	static FMaterialInvalidationSignature OnMaterialInvalidated;
+	static FSimpleMulticastDelegate OnSkyAtmosphereModified;
+	
 	static const class FMeshMapBuildData* GetPrimitiveMeshMapBuildData(const UPrimitiveComponent* Component, int32 LODIndex = 0);
 	static const class FLightComponentMapBuildData* GetLightComponentMapBuildData(const ULightComponent* Component);
 	static const class FPrecomputedVolumetricLightmap* GetPrecomputedVolumetricLightmap(UWorld* World);
+	
 	static void EditorTick();
 	static void GameTick(float DeltaSeconds);
 	static bool IsStaticLightingSystemRunning();
