GPULM: Implement preliminary support for static shadow depth maps to support volumetric fog shadows for stationary lights. Code duplication will be removed in following TCompactingObjectPool refactor.

diff --git a/Engine/Plugins/Experimental/GPULightmass/Shaders/Private/StaticShadowDepthMap.usf b/Engine/Plugins/Experimental/GPULightmass/Shaders/Private/StaticShadowDepthMap.usf
new file mode 100644
--- /dev/null
+++ b/Engine/Plugins/Experimental/GPULightmass/Shaders/Private/StaticShadowDepthMap.usf
@@ -0,0 +1,74 @@
+// Copyright Epic Games, Inc. All Rights Reserved.
+
+#include "/Engine/Private/RayTracing/RayTracingCommon.ush"
+
+RaytracingAccelerationStructure TLAS;
+
+int StaticShadowDepthMapSuperSampleFactor;
+int2 ShadowMapSize;
+float4x4 LightToWorld;
+float4x4 WorldToLight;
+float3 LightSpaceImportanceBoundsMin;
+float3 LightSpaceImportanceBoundsMax;
+float MaxPossibleDistance;
+RWTexture2D<half> DepthMapTexture;
+
+[shader("raygeneration")]
+void StaticShadowDepthMapTracingRG()
+{
+	const float2 Fraction = float2(DispatchRaysIndex().xy) / ((ShadowMapSize - 1) * StaticShadowDepthMapSuperSampleFactor);
+#if LIGHT_TYPE == 0
+	const float4 LightSpaceStartPosition = float4(
+		LightSpaceImportanceBoundsMin.xy + Fraction * (LightSpaceImportanceBoundsMax.xy - LightSpaceImportanceBoundsMin.xy),
+		LightSpaceImportanceBoundsMin.z,
+		1);
+	const float4 LightSpaceEndPosition = float4(LightSpaceStartPosition.x, LightSpaceStartPosition.y, LightSpaceImportanceBoundsMax.z, 1);
+#elif LIGHT_TYPE == 1
+	const float4 LightSpaceStartPosition = float4(0, 0, 0, 1);
+	const float4 LightSpaceEndPosition = float4(
+		LightSpaceImportanceBoundsMin.xy + Fraction * (LightSpaceImportanceBoundsMax.xy - LightSpaceImportanceBoundsMin.xy),
+		LightSpaceImportanceBoundsMax.z,
+		1);
+#else
+	const float Phi = Fraction.y * PI;
+	const float SinPhi = sin(Phi);
+	const float Theta = Fraction.x * 2 * PI;
+	const float3 Direction = normalize(float3(cos(Theta) * SinPhi, sin(Theta) * SinPhi, cos(Phi)));
+	
+	const float4 LightSpaceStartPosition = float4(0, 0, 0, 1);
+	const float4 LightSpaceEndPosition = float4(Direction * MaxPossibleDistance, 1);
+#endif
+	
+	const float3 WorldSpaceStartPosition = mul(LightSpaceStartPosition, LightToWorld).xyz;
+	const float3 WorldSpaceEndPosition = mul(LightSpaceEndPosition, LightToWorld).xyz;
+
+	RayDesc Ray;
+	Ray.Origin = WorldSpaceStartPosition;
+	Ray.Direction = normalize(WorldSpaceEndPosition - WorldSpaceStartPosition);
+	Ray.TMin = 0;
+	Ray.TMax = length(WorldSpaceEndPosition - WorldSpaceStartPosition);
+
+	const uint RayFlags = 0;
+	
+	FMinimalPayload MinimalPayload = TraceVisibilityRay(
+		TLAS,
+		RayFlags,
+		RAY_TRACING_MASK_SHADOW,
+		uint2(0, 0),
+		Ray);
+
+	if (MinimalPayload.IsHit())
+	{
+	#if LIGHT_TYPE == 0 || LIGHT_TYPE == 2
+		DepthMapTexture[DispatchRaysIndex().xy] = MinimalPayload.HitT / MaxPossibleDistance;
+	#elif LIGHT_TYPE == 1
+		const float4 WorldSpaceHitPosition = float4(Ray.Origin + Ray.Direction * MinimalPayload.HitT, 1);
+		const float4 LightSpaceHitPosition = mul(WorldSpaceHitPosition, WorldToLight);
+		DepthMapTexture[DispatchRaysIndex().xy] = LightSpaceHitPosition.z / MaxPossibleDistance;
+	#endif
+	}
+	else
+	{
+		DepthMapTexture[DispatchRaysIndex().xy] = 1;
+	}
+}
diff --git a/Engine/Plugins/Experimental/GPULightmass/Source/GPULightmass/Private/LightmapRayTracing.cpp b/Engine/Plugins/Experimental/GPULightmass/Source/GPULightmass/Private/LightmapRayTracing.cpp
--- a/Engine/Plugins/Experimental/GPULightmass/Source/GPULightmass/Private/LightmapRayTracing.cpp
+++ b/Engine/Plugins/Experimental/GPULightmass/Source/GPULightmass/Private/LightmapRayTracing.cpp
@@ -12,6 +12,7 @@ bool IsSupportedVertexFactoryType(const FVertexFactoryType* VertexFactoryType)
 IMPLEMENT_GLOBAL_SHADER(FLightmapPathTracingRGS, "/Plugin/GPULightmass/Private/LightmapPathTracing.usf", "LightmapPathTracingMainRG", SF_RayGen);
 IMPLEMENT_GLOBAL_SHADER(FVolumetricLightmapPathTracingRGS, "/Plugin/GPULightmass/Private/LightmapPathTracing.usf", "VolumetricLightmapPathTracingMainRG", SF_RayGen);
 IMPLEMENT_GLOBAL_SHADER(FStationaryLightShadowTracingRGS, "/Plugin/GPULightmass/Private/LightmapPathTracing.usf", "StationaryLightShadowTracingMainRG", SF_RayGen);
+IMPLEMENT_GLOBAL_SHADER(FStaticShadowDepthMapTracingRGS, "/Plugin/GPULightmass/Private/StaticShadowDepthMap.usf", "StaticShadowDepthMapTracingRG", SF_RayGen);
 
 IMPLEMENT_GLOBAL_SHADER(FFirstBounceRayGuidingCDFBuildCS, "/Plugin/GPULightmass/Private/FirstBounceRayGuidingCDFBuild.usf", "FirstBounceRayGuidingCDFBuildCS", SF_Compute);
 
diff --git a/Engine/Plugins/Experimental/GPULightmass/Source/GPULightmass/Private/LightmapRayTracing.h b/Engine/Plugins/Experimental/GPULightmass/Source/GPULightmass/Private/LightmapRayTracing.h
--- a/Engine/Plugins/Experimental/GPULightmass/Source/GPULightmass/Private/LightmapRayTracing.h
+++ b/Engine/Plugins/Experimental/GPULightmass/Source/GPULightmass/Private/LightmapRayTracing.h
@@ -223,4 +223,37 @@ class FFirstBounceRayGuidingCDFBuildCS : public FGlobalShader
 	END_SHADER_PARAMETER_STRUCT()
 };
 
+class FStaticShadowDepthMapTracingRGS : public FGlobalShader
+{
+	DECLARE_GLOBAL_SHADER(FStaticShadowDepthMapTracingRGS)
+	SHADER_USE_ROOT_PARAMETER_STRUCT(FStaticShadowDepthMapTracingRGS, FGlobalShader)
+
+	class FLightType : SHADER_PERMUTATION_INT("LIGHT_TYPE", 3);
+
+	using FPermutationDomain = TShaderPermutationDomain<FLightType>;
+	
+	static bool ShouldCompilePermutation(const FGlobalShaderPermutationParameters& Parameters)
+	{
+		return EnumHasAllFlags(Parameters.Flags, EShaderPermutationFlags::HasEditorOnlyData) && ShouldCompileRayTracingShadersForProject(Parameters.Platform);
+	}
+
+	static void ModifyCompilationEnvironment(const FGlobalShaderPermutationParameters& Parameters, FShaderCompilerEnvironment& OutEnvironment)
+	{
+		OutEnvironment.SetDefine(TEXT("GPreviewLightmapPhysicalTileSize"), GPreviewLightmapPhysicalTileSize);
+		OutEnvironment.CompilerFlags.Add(CFLAG_WarningsAsErrors);
+	}
+
+	BEGIN_SHADER_PARAMETER_STRUCT(FParameters, )
+		SHADER_PARAMETER(int32, StaticShadowDepthMapSuperSampleFactor)
+		SHADER_PARAMETER(FIntPoint, ShadowMapSize)
+		SHADER_PARAMETER(FMatrix44f, LightToWorld)
+		SHADER_PARAMETER(FMatrix44f, WorldToLight)
+		SHADER_PARAMETER(FVector3f, LightSpaceImportanceBoundsMin)
+		SHADER_PARAMETER(FVector3f, LightSpaceImportanceBoundsMax)
+		SHADER_PARAMETER(float, MaxPossibleDistance)
+		SHADER_PARAMETER_SRV(RaytracingAccelerationStructure, TLAS)
+		SHADER_PARAMETER_RDG_TEXTURE_UAV(RWTexture2D<half>, DepthMapTexture)
+	END_SHADER_PARAMETER_STRUCT()
+};
+
 #endif
\ No newline at end of file
diff --git a/Engine/Plugins/Experimental/GPULightmass/Source/GPULightmass/Private/LightmapRenderer.cpp b/Engine/Plugins/Experimental/GPULightmass/Source/GPULightmass/Private/LightmapRenderer.cpp
--- a/Engine/Plugins/Experimental/GPULightmass/Source/GPULightmass/Private/LightmapRenderer.cpp
+++ b/Engine/Plugins/Experimental/GPULightmass/Source/GPULightmass/Private/LightmapRenderer.cpp
@@ -1223,6 +1223,14 @@ bool FSceneRenderState::SetupRayTracingScene(int32 LODIndex)
 			}
 			{
 				RayGenShaderTable.Add(GlobalShaderMap->GetShader<FStationaryLightShadowTracingRGS>().GetRayTracingShader());
+
+				FStaticShadowDepthMapTracingRGS::FPermutationDomain PermutationVector;
+				PermutationVector.Set<FStaticShadowDepthMapTracingRGS::FLightType>(0);
+				RayGenShaderTable.Add(GlobalShaderMap->GetShader<FStaticShadowDepthMapTracingRGS>(PermutationVector).GetRayTracingShader());
+				PermutationVector.Set<FStaticShadowDepthMapTracingRGS::FLightType>(1);
+				RayGenShaderTable.Add(GlobalShaderMap->GetShader<FStaticShadowDepthMapTracingRGS>(PermutationVector).GetRayTracingShader());
+				PermutationVector.Set<FStaticShadowDepthMapTracingRGS::FLightType>(2);
+				RayGenShaderTable.Add(GlobalShaderMap->GetShader<FStaticShadowDepthMapTracingRGS>(PermutationVector).GetRayTracingShader());
 			}
 			{
 				FVolumetricLightmapPathTracingRGS::FPermutationDomain PermutationVector;
diff --git a/Engine/Plugins/Experimental/GPULightmass/Source/GPULightmass/Private/Scene/Lights.cpp b/Engine/Plugins/Experimental/GPULightmass/Source/GPULightmass/Private/Scene/Lights.cpp
--- a/Engine/Plugins/Experimental/GPULightmass/Source/GPULightmass/Private/Scene/Lights.cpp
+++ b/Engine/Plugins/Experimental/GPULightmass/Source/GPULightmass/Private/Scene/Lights.cpp
@@ -1,11 +1,15 @@
 // Copyright Epic Games, Inc. All Rights Reserved.
 
 #include "Lights.h"
+
+#include "LightmapRayTracing.h"
+#include "Scene.h"
 #include "RenderGraphBuilder.h"
 #include "ReflectionEnvironment.h"
 #include "RectLightTextureManager.h"
 #include "Components/SkyAtmosphereComponent.h"
 #include "UObject/UObjectIterator.h"
+#include "RenderGraphUtils.h"
 
 RENDERER_API void PrepareSkyTexture_Internal(
 	FRDGBuilder& GraphBuilder,
@@ -356,4 +360,393 @@ void FSkyLightRenderState::PrepareSkyTexture(FRHICommandListImmediate& RHICmdLis
 	GraphBuilder.Execute();
 }
 
+void FDirectionalLightRenderState::RenderStaticShadowDepthMap(FRHICommandListImmediate& RHICmdList, FSceneRenderState& Scene)
+{
+	Scene.SetupRayTracingScene();
+	
+	FVector3f XAxis, YAxis;
+	FVector3f(Direction).FindBestAxisVectors(XAxis, YAxis);
+	// Create a coordinate system for the dominant directional light, with the z axis corresponding to the light's direction
+	FMatrix44f WorldToLight = FBasisVectorMatrix44f(XAxis, YAxis, FVector3f(Direction), FVector4f(0, 0, 0));
+
+	const FBox3f LightSpaceImportanceBounds = FBox3f{Scene.CombinedImportanceVolume}.TransformBy(WorldToLight);
+
+	const float ClampedResolutionScale = 1.0f;
+	const float StaticShadowDepthMapTransitionSampleDistanceX = 10;
+	const float StaticShadowDepthMapTransitionSampleDistanceY = 10;
+	
+	int32 ShadowMapSizeX = FMath::TruncToInt(FMath::Max(LightSpaceImportanceBounds.GetExtent().X * 2.0f * ClampedResolutionScale / StaticShadowDepthMapTransitionSampleDistanceX, 4.0f));
+	int32 ShadowMapSizeY = FMath::TruncToInt(FMath::Max(LightSpaceImportanceBounds.GetExtent().Y * 2.0f * ClampedResolutionScale / StaticShadowDepthMapTransitionSampleDistanceY, 4.0f));
+
+	const uint64 StaticShadowDepthMapMaxSamples = 16777216;
+	
+	// Clamp the number of dominant shadow samples generated if necessary while maintaining aspect ratio
+	if ((uint64)ShadowMapSizeX * (uint64)ShadowMapSizeY > StaticShadowDepthMapMaxSamples)
+	{
+		const float AspectRatio = ShadowMapSizeX / (float)ShadowMapSizeY;
+		ShadowMapSizeY = FMath::TruncToInt(FMath::Sqrt(StaticShadowDepthMapMaxSamples / AspectRatio));
+		ShadowMapSizeX = FMath::TruncToInt(StaticShadowDepthMapMaxSamples / (float)ShadowMapSizeY);
+	}
+
+	FRDGBuilder GraphBuilder(RHICmdList);
+	
+	FRDGTextureRef DepthMapTexture = GraphBuilder.CreateTexture(
+		FRDGTextureDesc::Create2D(
+			FIntPoint{ShadowMapSizeX, ShadowMapSizeY},
+			PF_R16F,
+			FClearValueBinding::DepthFar,
+			ETextureCreateFlags::ShaderResource | ETextureCreateFlags::RenderTargetable | ETextureCreateFlags::UAV
+		), TEXT("GPULMStaticShadowDepthMap"));
+
+	FStaticShadowDepthMapTracingRGS::FParameters* PassParameters = GraphBuilder.AllocParameters<FStaticShadowDepthMapTracingRGS::FParameters>();
+	PassParameters->TLAS = Scene.RayTracingSceneSRV;
+	PassParameters->ShadowMapSize = FIntPoint{ShadowMapSizeX, ShadowMapSizeY};
+	PassParameters->StaticShadowDepthMapSuperSampleFactor = 1;
+	PassParameters->LightSpaceImportanceBoundsMin = LightSpaceImportanceBounds.Min;
+	PassParameters->LightSpaceImportanceBoundsMax = LightSpaceImportanceBounds.Max;
+	PassParameters->LightToWorld = WorldToLight.InverseFast();
+	PassParameters->MaxPossibleDistance = LightSpaceImportanceBounds.Max.Z - LightSpaceImportanceBounds.Min.Z;
+	PassParameters->DepthMapTexture = GraphBuilder.CreateUAV(DepthMapTexture);
+
+	FStaticShadowDepthMapTracingRGS::FPermutationDomain PermutationVector;
+	PermutationVector.Set<FStaticShadowDepthMapTracingRGS::FLightType>(0);
+	auto RayGenShader = GetGlobalShaderMap(GMaxRHIFeatureLevel)->GetShader<FStaticShadowDepthMapTracingRGS>(PermutationVector);
+	ClearUnusedGraphResources(RayGenShader, PassParameters);
+	
+	GraphBuilder.AddPass(
+		RDG_EVENT_NAME("StaticShadowDepthMapTracing"),
+		PassParameters,
+		ERDGPassFlags::Compute,
+		[PassParameters, RayGenShader, RayTracingSceneRHI = Scene.RayTracingScene, RayTracingPipelineState = Scene.RayTracingPipelineState](FRHIRayTracingCommandList& RHICmdList)
+	{
+		FRayTracingShaderBindingsWriter GlobalResources;
+		SetShaderParameters(GlobalResources, RayGenShader, *PassParameters);
+
+		RHICmdList.RayTraceDispatch(RayTracingPipelineState, RayGenShader.GetRayTracingShader(), RayTracingSceneRHI, GlobalResources, PassParameters->ShadowMapSize.X, PassParameters->ShadowMapSize.Y);
+	}
+	);
+
+	FRHIGPUTextureReadback DepthMapTextureReadback(TEXT("GPULMStaticShadowDepthMapReadback"));
+	
+	AddEnqueueCopyPass(GraphBuilder, &DepthMapTextureReadback, DepthMapTexture);
+	
+	GraphBuilder.Execute();
+
+	RHICmdList.BlockUntilGPUIdle();
+
+	check(DepthMapTextureReadback.IsReady());
+
+	int32 RowPitchInPixels;
+	FFloat16* LockedData = (FFloat16*)DepthMapTextureReadback.Lock(RowPitchInPixels);
+	LightComponentMapBuildData->DepthMap.Empty();
+	LightComponentMapBuildData->DepthMap.DepthSamples.AddZeroed(ShadowMapSizeX * ShadowMapSizeY);
+	for (int32 Y = 0; Y < ShadowMapSizeY; Y++)
+	{
+		for (int32 X = 0; X < ShadowMapSizeX; X++)
+		{
+			LightComponentMapBuildData->DepthMap.DepthSamples[Y * ShadowMapSizeX + X] = LockedData[Y * RowPitchInPixels + X];
+		}
+	}
+	LightComponentMapBuildData->DepthMap.ShadowMapSizeX = ShadowMapSizeX;
+	LightComponentMapBuildData->DepthMap.ShadowMapSizeY = ShadowMapSizeY;
+	WorldToLight *= FTranslationMatrix44f(-LightSpaceImportanceBounds.Min) * FScaleMatrix44f(FVector3f(1.0f) / (LightSpaceImportanceBounds.Max - LightSpaceImportanceBounds.Min));
+	LightComponentMapBuildData->DepthMap.WorldToLight = FMatrix(WorldToLight);
+	DepthMapTextureReadback.Unlock();
+	
+	Scene.DestroyRayTracingScene();
+}
+
+void FSpotLightRenderState::RenderStaticShadowDepthMap(FRHICommandListImmediate& RHICmdList, FSceneRenderState& Scene)
+{
+	Scene.SetupRayTracingScene();
+	
+	FVector3f XAxis, YAxis;
+	FVector3f(Direction).FindBestAxisVectors(XAxis, YAxis);
+	// Create a coordinate system for the dominant directional light, with the z axis corresponding to the light's direction
+	FMatrix44f WorldToLight = FTranslationMatrix44f(-FVector3f(Position)) * FBasisVectorMatrix44f(XAxis, YAxis, FVector3f(Direction), FVector4f(0, 0, 0));
+
+	// Distance from the light's direction axis to the edge of the cone at the radius of the light
+	const float HalfCrossSectionLength = AttenuationRadius * FMath::Tan(FMath::Acos(SpotAngles.X));
+	
+	const FBox3f LightSpaceImportanceBounds
+	{
+		FVector3f{-HalfCrossSectionLength, -HalfCrossSectionLength, 0},
+		FVector3f{HalfCrossSectionLength, HalfCrossSectionLength, AttenuationRadius},
+	};
+	
+	const float ClampedResolutionScale = 1.0f;
+	const float StaticShadowDepthMapTransitionSampleDistanceX = 10;
+	const float StaticShadowDepthMapTransitionSampleDistanceY = 10;
+	
+	int32 ShadowMapSizeX = FMath::TruncToInt(FMath::Max(HalfCrossSectionLength * ClampedResolutionScale / StaticShadowDepthMapTransitionSampleDistanceX, 4.0f));
+	int32 ShadowMapSizeY = ShadowMapSizeX;
+
+	const uint64 StaticShadowDepthMapMaxSamples = 16777216;
+	
+	// Clamp the number of dominant shadow samples generated if necessary while maintaining aspect ratio
+	if ((uint64)ShadowMapSizeX * (uint64)ShadowMapSizeY > StaticShadowDepthMapMaxSamples)
+	{
+		const float AspectRatio = ShadowMapSizeX / (float)ShadowMapSizeY;
+		ShadowMapSizeY = FMath::TruncToInt(FMath::Sqrt(StaticShadowDepthMapMaxSamples / AspectRatio));
+		ShadowMapSizeX = FMath::TruncToInt(StaticShadowDepthMapMaxSamples / (float)ShadowMapSizeY);
+	}
+
+	FRDGBuilder GraphBuilder(RHICmdList);
+	
+	FRDGTextureRef DepthMapTexture = GraphBuilder.CreateTexture(
+		FRDGTextureDesc::Create2D(
+			FIntPoint{ShadowMapSizeX, ShadowMapSizeY},
+			PF_R16F,
+			FClearValueBinding::DepthFar,
+			ETextureCreateFlags::ShaderResource | ETextureCreateFlags::RenderTargetable | ETextureCreateFlags::UAV
+		), TEXT("GPULMStaticShadowDepthMap"));
+
+	const float MaxPossibleDistance = LightSpaceImportanceBounds.Max.Z - LightSpaceImportanceBounds.Min.Z;
+	const FMatrix44f LightToWorld = WorldToLight.InverseFast();
+
+	FStaticShadowDepthMapTracingRGS::FParameters* PassParameters = GraphBuilder.AllocParameters<FStaticShadowDepthMapTracingRGS::FParameters>();
+	PassParameters->TLAS = Scene.RayTracingSceneSRV;
+	PassParameters->ShadowMapSize = FIntPoint{ShadowMapSizeX, ShadowMapSizeY};
+	PassParameters->StaticShadowDepthMapSuperSampleFactor = 1;
+	PassParameters->LightSpaceImportanceBoundsMin = LightSpaceImportanceBounds.Min;
+	PassParameters->LightSpaceImportanceBoundsMax = LightSpaceImportanceBounds.Max;
+	PassParameters->LightToWorld = LightToWorld;
+	PassParameters->WorldToLight = WorldToLight;
+	PassParameters->MaxPossibleDistance = MaxPossibleDistance;
+	PassParameters->DepthMapTexture = GraphBuilder.CreateUAV(DepthMapTexture);
+
+	FStaticShadowDepthMapTracingRGS::FPermutationDomain PermutationVector;
+	PermutationVector.Set<FStaticShadowDepthMapTracingRGS::FLightType>(1);
+	auto RayGenShader = GetGlobalShaderMap(GMaxRHIFeatureLevel)->GetShader<FStaticShadowDepthMapTracingRGS>(PermutationVector);
+	ClearUnusedGraphResources(RayGenShader, PassParameters);
+	
+	GraphBuilder.AddPass(
+		RDG_EVENT_NAME("StaticShadowDepthMapTracing"),
+		PassParameters,
+		ERDGPassFlags::Compute,
+		[PassParameters, RayGenShader, RayTracingSceneRHI = Scene.RayTracingScene, RayTracingPipelineState = Scene.RayTracingPipelineState](FRHIRayTracingCommandList& RHICmdList)
+	{
+		FRayTracingShaderBindingsWriter GlobalResources;
+		SetShaderParameters(GlobalResources, RayGenShader, *PassParameters);
+
+		RHICmdList.RayTraceDispatch(RayTracingPipelineState, RayGenShader.GetRayTracingShader(), RayTracingSceneRHI, GlobalResources, PassParameters->ShadowMapSize.X, PassParameters->ShadowMapSize.Y);
+	}
+	);
+
+	FRHIGPUTextureReadback DepthMapTextureReadback(TEXT("GPULMStaticShadowDepthMapReadback"));
+	
+	AddEnqueueCopyPass(GraphBuilder, &DepthMapTextureReadback, DepthMapTexture);
+	
+	GraphBuilder.Execute();
+
+	RHICmdList.BlockUntilGPUIdle();
+
+	check(DepthMapTextureReadback.IsReady());
+
+	int32 RowPitchInPixels;
+	FFloat16* LockedData = (FFloat16*)DepthMapTextureReadback.Lock(RowPitchInPixels);
+	LightComponentMapBuildData->DepthMap.Empty();
+	LightComponentMapBuildData->DepthMap.DepthSamples.AddZeroed(ShadowMapSizeX * ShadowMapSizeY);
+	for (int32 Y = 0; Y < ShadowMapSizeY; Y++)
+	{
+		for (int32 X = 0; X < ShadowMapSizeX; X++)
+		{
+			LightComponentMapBuildData->DepthMap.DepthSamples[Y * ShadowMapSizeX + X] = LockedData[Y * RowPitchInPixels + X];
+		}
+	}
+	LightComponentMapBuildData->DepthMap.ShadowMapSizeX = ShadowMapSizeX;
+	LightComponentMapBuildData->DepthMap.ShadowMapSizeY = ShadowMapSizeY;
+	WorldToLight *=
+		// Perspective projection sized to the spotlight cone
+		FPerspectiveMatrix44f(FMath::Acos(SpotAngles.X), 1, 1, 0, AttenuationRadius)
+		// Convert from NDC to texture space, normalize Z
+		* FMatrix44f(
+			FPlane4f(.5f,	0,		0,											0),
+			FPlane4f(0,	.5f,	0,											0),
+			FPlane4f(0,	0,		1.0f / LightSpaceImportanceBounds.Max.Z,	0),
+			FPlane4f(.5f,	.5f,	0,											1));
+	LightComponentMapBuildData->DepthMap.WorldToLight = FMatrix(WorldToLight);
+	DepthMapTextureReadback.Unlock();
+	
+	Scene.DestroyRayTracingScene();
+}
+
+void FPointLightRenderState::RenderStaticShadowDepthMap(FRHICommandListImmediate& RHICmdList, FSceneRenderState& Scene)
+{
+	Scene.SetupRayTracingScene();
+
+	const float ClampedResolutionScale = 1.0f;
+	const float StaticShadowDepthMapTransitionSampleDistanceX = 10;
+	const float StaticShadowDepthMapTransitionSampleDistanceY = 10;
+	
+	int32 ShadowMapSizeX = FMath::TruncToInt(FMath::Max(AttenuationRadius * 4 * ClampedResolutionScale / StaticShadowDepthMapTransitionSampleDistanceX, 4.0f));
+	int32 ShadowMapSizeY = ShadowMapSizeX;
+
+	const uint64 StaticShadowDepthMapMaxSamples = 16777216;
+	
+	// Clamp the number of dominant shadow samples generated if necessary while maintaining aspect ratio
+	if ((uint64)ShadowMapSizeX * (uint64)ShadowMapSizeY > StaticShadowDepthMapMaxSamples)
+	{
+		const float AspectRatio = ShadowMapSizeX / (float)ShadowMapSizeY;
+		ShadowMapSizeY = FMath::TruncToInt(FMath::Sqrt(StaticShadowDepthMapMaxSamples / AspectRatio));
+		ShadowMapSizeX = FMath::TruncToInt(StaticShadowDepthMapMaxSamples / (float)ShadowMapSizeY);
+	}
+
+	FRDGBuilder GraphBuilder(RHICmdList);
+	
+	FRDGTextureRef DepthMapTexture = GraphBuilder.CreateTexture(
+		FRDGTextureDesc::Create2D(
+			FIntPoint{ShadowMapSizeX, ShadowMapSizeY},
+			PF_R16F,
+			FClearValueBinding::DepthFar,
+			ETextureCreateFlags::ShaderResource | ETextureCreateFlags::RenderTargetable | ETextureCreateFlags::UAV
+		), TEXT("GPULMStaticShadowDepthMap"));
+
+	FStaticShadowDepthMapTracingRGS::FParameters* PassParameters = GraphBuilder.AllocParameters<FStaticShadowDepthMapTracingRGS::FParameters>();
+	PassParameters->TLAS = Scene.RayTracingSceneSRV;
+	PassParameters->ShadowMapSize = FIntPoint{ShadowMapSizeX, ShadowMapSizeY};
+	PassParameters->StaticShadowDepthMapSuperSampleFactor = 1;
+	PassParameters->LightSpaceImportanceBoundsMin = FVector3f{EForceInit::ForceInitToZero};
+	PassParameters->LightSpaceImportanceBoundsMax = FVector3f{EForceInit::ForceInitToZero};
+	PassParameters->LightToWorld = FTranslationMatrix44f(FVector3f(Position));
+	PassParameters->WorldToLight = FTranslationMatrix44f(-FVector3f(Position));
+	PassParameters->MaxPossibleDistance = AttenuationRadius;
+	PassParameters->DepthMapTexture = GraphBuilder.CreateUAV(DepthMapTexture);
+
+	FStaticShadowDepthMapTracingRGS::FPermutationDomain PermutationVector;
+	PermutationVector.Set<FStaticShadowDepthMapTracingRGS::FLightType>(2);
+	auto RayGenShader = GetGlobalShaderMap(GMaxRHIFeatureLevel)->GetShader<FStaticShadowDepthMapTracingRGS>(PermutationVector);
+	ClearUnusedGraphResources(RayGenShader, PassParameters);
+	
+	GraphBuilder.AddPass(
+		RDG_EVENT_NAME("StaticShadowDepthMapTracing"),
+		PassParameters,
+		ERDGPassFlags::Compute,
+		[PassParameters, RayGenShader, RayTracingSceneRHI = Scene.RayTracingScene, RayTracingPipelineState = Scene.RayTracingPipelineState](FRHIRayTracingCommandList& RHICmdList)
+	{
+		FRayTracingShaderBindingsWriter GlobalResources;
+		SetShaderParameters(GlobalResources, RayGenShader, *PassParameters);
+
+		RHICmdList.RayTraceDispatch(RayTracingPipelineState, RayGenShader.GetRayTracingShader(), RayTracingSceneRHI, GlobalResources, PassParameters->ShadowMapSize.X, PassParameters->ShadowMapSize.Y);
+	}
+	);
+
+	FRHIGPUTextureReadback DepthMapTextureReadback(TEXT("GPULMStaticShadowDepthMapReadback"));
+	
+	AddEnqueueCopyPass(GraphBuilder, &DepthMapTextureReadback, DepthMapTexture);
+	
+	GraphBuilder.Execute();
+
+	RHICmdList.BlockUntilGPUIdle();
+
+	check(DepthMapTextureReadback.IsReady());
+
+	int32 RowPitchInPixels;
+	FFloat16* LockedData = (FFloat16*)DepthMapTextureReadback.Lock(RowPitchInPixels);
+	LightComponentMapBuildData->DepthMap.Empty();
+	LightComponentMapBuildData->DepthMap.DepthSamples.AddZeroed(ShadowMapSizeX * ShadowMapSizeY);
+	for (int32 Y = 0; Y < ShadowMapSizeY; Y++)
+	{
+		for (int32 X = 0; X < ShadowMapSizeX; X++)
+		{
+			LightComponentMapBuildData->DepthMap.DepthSamples[Y * ShadowMapSizeX + X] = LockedData[Y * RowPitchInPixels + X];
+		}
+	}
+	LightComponentMapBuildData->DepthMap.ShadowMapSizeX = ShadowMapSizeX;
+	LightComponentMapBuildData->DepthMap.ShadowMapSizeY = ShadowMapSizeY;
+	LightComponentMapBuildData->DepthMap.WorldToLight = FMatrix::Identity;
+	DepthMapTextureReadback.Unlock();
+	
+	Scene.DestroyRayTracingScene();
+}
+
+void FRectLightRenderState::RenderStaticShadowDepthMap(FRHICommandListImmediate& RHICmdList, FSceneRenderState& Scene)
+{
+	Scene.SetupRayTracingScene();
+
+	const float ClampedResolutionScale = 1.0f;
+	const float StaticShadowDepthMapTransitionSampleDistanceX = 10;
+	const float StaticShadowDepthMapTransitionSampleDistanceY = 10;
+	
+	int32 ShadowMapSizeX = FMath::TruncToInt(FMath::Max(AttenuationRadius * 4 * ClampedResolutionScale / StaticShadowDepthMapTransitionSampleDistanceX, 4.0f));
+	int32 ShadowMapSizeY = ShadowMapSizeX;
+
+	const uint64 StaticShadowDepthMapMaxSamples = 16777216;
+	
+	// Clamp the number of dominant shadow samples generated if necessary while maintaining aspect ratio
+	if ((uint64)ShadowMapSizeX * (uint64)ShadowMapSizeY > StaticShadowDepthMapMaxSamples)
+	{
+		const float AspectRatio = ShadowMapSizeX / (float)ShadowMapSizeY;
+		ShadowMapSizeY = FMath::TruncToInt(FMath::Sqrt(StaticShadowDepthMapMaxSamples / AspectRatio));
+		ShadowMapSizeX = FMath::TruncToInt(StaticShadowDepthMapMaxSamples / (float)ShadowMapSizeY);
+	}
+
+	FRDGBuilder GraphBuilder(RHICmdList);
+	
+	FRDGTextureRef DepthMapTexture = GraphBuilder.CreateTexture(
+		FRDGTextureDesc::Create2D(
+			FIntPoint{ShadowMapSizeX, ShadowMapSizeY},
+			PF_R16F,
+			FClearValueBinding::DepthFar,
+			ETextureCreateFlags::ShaderResource | ETextureCreateFlags::RenderTargetable | ETextureCreateFlags::UAV
+		), TEXT("GPULMStaticShadowDepthMap"));
+
+	FStaticShadowDepthMapTracingRGS::FParameters* PassParameters = GraphBuilder.AllocParameters<FStaticShadowDepthMapTracingRGS::FParameters>();
+	PassParameters->TLAS = Scene.RayTracingSceneSRV;
+	PassParameters->ShadowMapSize = FIntPoint{ShadowMapSizeX, ShadowMapSizeY};
+	PassParameters->StaticShadowDepthMapSuperSampleFactor = 1;
+	PassParameters->LightSpaceImportanceBoundsMin = FVector3f{EForceInit::ForceInitToZero};
+	PassParameters->LightSpaceImportanceBoundsMax = FVector3f{EForceInit::ForceInitToZero};
+	PassParameters->LightToWorld = FTranslationMatrix44f(FVector3f(Position));
+	PassParameters->WorldToLight = FTranslationMatrix44f(-FVector3f(Position));
+	PassParameters->MaxPossibleDistance = AttenuationRadius;
+	PassParameters->DepthMapTexture = GraphBuilder.CreateUAV(DepthMapTexture);
+
+	FStaticShadowDepthMapTracingRGS::FPermutationDomain PermutationVector;
+	PermutationVector.Set<FStaticShadowDepthMapTracingRGS::FLightType>(2);
+	auto RayGenShader = GetGlobalShaderMap(GMaxRHIFeatureLevel)->GetShader<FStaticShadowDepthMapTracingRGS>(PermutationVector);
+	ClearUnusedGraphResources(RayGenShader, PassParameters);
+	
+	GraphBuilder.AddPass(
+		RDG_EVENT_NAME("StaticShadowDepthMapTracing"),
+		PassParameters,
+		ERDGPassFlags::Compute,
+		[PassParameters, RayGenShader, RayTracingSceneRHI = Scene.RayTracingScene, RayTracingPipelineState = Scene.RayTracingPipelineState](FRHIRayTracingCommandList& RHICmdList)
+	{
+		FRayTracingShaderBindingsWriter GlobalResources;
+		SetShaderParameters(GlobalResources, RayGenShader, *PassParameters);
+
+		RHICmdList.RayTraceDispatch(RayTracingPipelineState, RayGenShader.GetRayTracingShader(), RayTracingSceneRHI, GlobalResources, PassParameters->ShadowMapSize.X, PassParameters->ShadowMapSize.Y);
+	}
+	);
+
+	FRHIGPUTextureReadback DepthMapTextureReadback(TEXT("GPULMStaticShadowDepthMapReadback"));
+	
+	AddEnqueueCopyPass(GraphBuilder, &DepthMapTextureReadback, DepthMapTexture);
+	
+	GraphBuilder.Execute();
+
+	RHICmdList.BlockUntilGPUIdle();
+
+	check(DepthMapTextureReadback.IsReady());
+
+	int32 RowPitchInPixels;
+	FFloat16* LockedData = (FFloat16*)DepthMapTextureReadback.Lock(RowPitchInPixels);
+	LightComponentMapBuildData->DepthMap.Empty();
+	LightComponentMapBuildData->DepthMap.DepthSamples.AddZeroed(ShadowMapSizeX * ShadowMapSizeY);
+	for (int32 Y = 0; Y < ShadowMapSizeY; Y++)
+	{
+		for (int32 X = 0; X < ShadowMapSizeX; X++)
+		{
+			LightComponentMapBuildData->DepthMap.DepthSamples[Y * ShadowMapSizeX + X] = LockedData[Y * RowPitchInPixels + X];
+		}
+	}
+	LightComponentMapBuildData->DepthMap.ShadowMapSizeX = ShadowMapSizeX;
+	LightComponentMapBuildData->DepthMap.ShadowMapSizeY = ShadowMapSizeY;
+	LightComponentMapBuildData->DepthMap.WorldToLight = FMatrix::Identity;
+	DepthMapTextureReadback.Unlock();
+	
+	Scene.DestroyRayTracingScene();
+}
+
 }
diff --git a/Engine/Plugins/Experimental/GPULightmass/Source/GPULightmass/Private/Scene/Lights.h b/Engine/Plugins/Experimental/GPULightmass/Source/GPULightmass/Private/Scene/Lights.h
--- a/Engine/Plugins/Experimental/GPULightmass/Source/GPULightmass/Private/Scene/Lights.h
+++ b/Engine/Plugins/Experimental/GPULightmass/Source/GPULightmass/Private/Scene/Lights.h
@@ -97,6 +97,7 @@ struct FLocalLightRenderState
 	TSharedPtr<FLightComponentMapBuildData, ESPMode::ThreadSafe> LightComponentMapBuildData;
 
 	virtual FLightRenderParameters GetLightShaderParameters() const = 0;
+	virtual void RenderStaticShadowDepthMap(FRHICommandListImmediate& RHICmdList, class FSceneRenderState& RenderState) {}
 };
 
 class FLightArrayBase;
@@ -221,6 +222,7 @@ struct FDirectionalLightRenderState : public FLocalLightRenderState
 	float LightSourceSoftAngle;
 
 	virtual FLightRenderParameters GetLightShaderParameters() const override;
+	virtual void RenderStaticShadowDepthMap(FRHICommandListImmediate& RHICmdList, FSceneRenderState& RenderState) override;
 };
 
 using FDirectionalLightRenderStateRef = TEntityArray<FDirectionalLightRenderState>::EntityRefType;
@@ -242,6 +244,7 @@ struct FPointLightRenderState : public FLocalLightRenderState
 	FTexture* IESTexture;
 
 	virtual FLightRenderParameters GetLightShaderParameters() const override;
+	virtual void RenderStaticShadowDepthMap(FRHICommandListImmediate& RHICmdList, FSceneRenderState& RenderState) override;
 };
 
 using FPointLightRenderStateRef = TEntityArray<FPointLightRenderState>::EntityRefType;
@@ -263,8 +266,8 @@ struct FSpotLightRenderState : public FLocalLightRenderState
 	bool IsInverseSquared;
 	FTexture* IESTexture;
 
-
 	virtual FLightRenderParameters GetLightShaderParameters() const override;
+	virtual void RenderStaticShadowDepthMap(FRHICommandListImmediate& RHICmdList, FSceneRenderState& RenderState) override;
 };
 
 using FSpotLightRenderStateRef = TEntityArray<FSpotLightRenderState>::EntityRefType;
@@ -292,6 +295,7 @@ struct FRectLightRenderState : public FLocalLightRenderState
 	float RectLightAtlasMaxLevel;
 
 	virtual FLightRenderParameters GetLightShaderParameters() const override;
+	virtual void RenderStaticShadowDepthMap(FRHICommandListImmediate& RHICmdList, FSceneRenderState& RenderState) override;
 };
 
 using FRectLightRenderStateRef = TEntityArray<FRectLightRenderState>::EntityRefType;
diff --git a/Engine/Plugins/Experimental/GPULightmass/Source/GPULightmass/Private/Scene/Scene.cpp b/Engine/Plugins/Experimental/GPULightmass/Source/GPULightmass/Private/Scene/Scene.cpp
--- a/Engine/Plugins/Experimental/GPULightmass/Source/GPULightmass/Private/Scene/Scene.cpp
+++ b/Engine/Plugins/Experimental/GPULightmass/Source/GPULightmass/Private/Scene/Scene.cpp
@@ -190,13 +190,10 @@ void FScene::GatherImportanceVolumes()
 		CombinedImportanceVolume = ReasonableSceneBounds;
 		ImportanceVolumes.Add(ReasonableSceneBounds);
 	}
-
-	float TargetDetailCellSize = Settings->VolumetricLightmapDetailCellSize;
-
-	ENQUEUE_RENDER_COMMAND(UpdateVLMRendererVolume)([&RenderState = RenderState, CombinedImportanceVolume, ImportanceVolumes, TargetDetailCellSize](FRHICommandList&) mutable {
-		RenderState.VolumetricLightmapRenderer->CombinedImportanceVolume = CombinedImportanceVolume;
-		RenderState.VolumetricLightmapRenderer->ImportanceVolumes = ImportanceVolumes;
-		RenderState.VolumetricLightmapRenderer->TargetDetailCellSize = TargetDetailCellSize;
+	
+	ENQUEUE_RENDER_COMMAND(UpdateLIVs)([&RenderState = RenderState, CombinedImportanceVolume, ImportanceVolumes](FRHICommandList&) mutable {
+		RenderState.CombinedImportanceVolume = CombinedImportanceVolume;
+		RenderState.ImportanceVolumes = ImportanceVolumes;
 	});
 }
 
@@ -398,9 +395,9 @@ void FScene::AddLight(LightComponentType* LightComponent)
 	{
 		FGeometry& Geometry = GeomIt.GetGeometry();
 
-		if (Light.AffectsBounds(Geometry.WorldBounds))
+		if (LightRef->AffectsBounds(Geometry.WorldBounds))
 		{
-			if (Light.CastsStationaryShadow())
+			if (LightRef->CastsStationaryShadow())
 			{
 				RelevantGeometriesToUpdateOnRenderThread.Add({ GeomIt.Index, GeomIt.Array.GetRenderStateArray() });
 			}
@@ -409,7 +406,7 @@ void FScene::AddLight(LightComponentType* LightComponent)
 			{
 				if (Lightmap.IsValid())
 				{
-					AddLightToLightmap(Lightmap.GetReference_Unsafe(), Light);
+					AddLightToLightmap(Lightmap.GetReference_Unsafe(), LightRef.GetReference_Unsafe());
 				}
 			}
 
@@ -420,6 +417,11 @@ void FScene::AddLight(LightComponentType* LightComponent)
 		}
 	}
 
+	if (LightRef->CastsStationaryShadow())
+	{
+		GatherImportanceVolumes();
+	}
+
 	ENQUEUE_RENDER_COMMAND(UpdateStaticLightingBufferCmd)(
 		[SceneProxiesToUpdateOnRenderThread = MoveTemp(SceneProxiesToUpdateOnRenderThread)](FRHICommandListImmediate& RHICmdList)
 	{
@@ -437,13 +439,23 @@ void FScene::AddLight(LightComponentType* LightComponent)
 		[
 			&RenderState = RenderState,
 			LightRenderState = MoveTemp(LightRenderState),
-			RelevantGeometriesToUpdateOnRenderThread
+			RelevantGeometriesToUpdateOnRenderThread,
+			bShouldRenderStaticShadowDepthMap = LightRef->CastsStationaryShadow(),
+			&StaticShadowDepthMap = LightComponent->StaticShadowDepthMap,
+			DepthMapPtr = &LightRef->LightComponentMapBuildData->DepthMap
 		](FRHICommandListImmediate& RHICmdList) mutable
 	{
 		typename LightTypeInfo<LightComponentType>::RenderStateRefType LightRenderStateRef = LightTypeInfo<LightComponentType>::GetLightRenderStateArray(RenderState.LightSceneRenderState).Emplace(MoveTemp(LightRenderState));
 
 		LightRenderStateRef->RenderThreadInit();
 
+		if (bShouldRenderStaticShadowDepthMap)
+		{
+			LightRenderStateRef->RenderStaticShadowDepthMap(RHICmdList, RenderState);
+			StaticShadowDepthMap.Data = DepthMapPtr;
+			StaticShadowDepthMap.InitRHI();
+		}
+
 		for (FGeometryRenderStateToken Token : RelevantGeometriesToUpdateOnRenderThread)
 		{
 			for (FLightmapRenderStateRef& Lightmap : Token.RenderStateArray->Get(Token.ElementId).LODLightmapRenderStates)
@@ -1524,7 +1536,10 @@ void FScene::BackgroundTick()
 			{
 				GatherImportanceVolumes();
 
-				ENQUEUE_RENDER_COMMAND(BackgroundTickRenderThread)([&RenderState = RenderState](FRHICommandListImmediate&) mutable {
+				float TargetDetailCellSize = Settings->VolumetricLightmapDetailCellSize;
+				
+				ENQUEUE_RENDER_COMMAND(BackgroundTickRenderThread)([&RenderState = RenderState, TargetDetailCellSize](FRHICommandListImmediate&) mutable {
+					RenderState.VolumetricLightmapRenderer->TargetDetailCellSize = TargetDetailCellSize;
 					RenderState.VolumetricLightmapRenderer->VoxelizeScene();
 					RenderState.VolumetricLightmapRenderer->FrameNumber = 0;
 					RenderState.VolumetricLightmapRenderer->SamplesTaken = 0;
@@ -1773,6 +1788,8 @@ void FLocalLightBuildInfo::AllocateMapBuildData(ULevel* StorageLevel)
 	UMapBuildDataRegistry* Registry = StorageLevel->GetOrCreateMapBuildData();
 	FLightComponentMapBuildData& LightBuildData = Registry->FindOrAllocateLightBuildData(GetComponentUObject()->LightGuid, true);
 	LightBuildData.ShadowMapChannel = CastsStationaryShadow() ? GetComponentUObject()->PreviewShadowMapChannel : INDEX_NONE;
+	// Copy to storage
+	LightBuildData.DepthMap = LightComponentMapBuildData->DepthMap;
 }
 
 void FScene::AddRelevantStaticLightGUIDs(FQuantizedLightmapData* QuantizedLightmapData, const FBoxSphereBounds& WorldBounds)
diff --git a/Engine/Plugins/Experimental/GPULightmass/Source/GPULightmass/Private/Scene/Scene.h b/Engine/Plugins/Experimental/GPULightmass/Source/GPULightmass/Private/Scene/Scene.h
--- a/Engine/Plugins/Experimental/GPULightmass/Source/GPULightmass/Private/Scene/Scene.h
+++ b/Engine/Plugins/Experimental/GPULightmass/Source/GPULightmass/Private/Scene/Scene.h
@@ -117,6 +117,9 @@ public:
 	TUniquePtr<FVolumetricLightmapRenderer> VolumetricLightmapRenderer;
 	TUniquePtr<FIrradianceCache> IrradianceCache;
 
+	FBox CombinedImportanceVolume;
+	TArray<FBox> ImportanceVolumes;
+	
 	ERHIFeatureLevel::Type FeatureLevel;
 
 	int32 GetPrimitiveIdForGPUScene(const FGeometryInstanceRenderStateRef& GeometryInstanceRef) const;
diff --git a/Engine/Plugins/Experimental/GPULightmass/Source/GPULightmass/Private/VolumetricLightmap.cpp b/Engine/Plugins/Experimental/GPULightmass/Source/GPULightmass/Private/VolumetricLightmap.cpp
--- a/Engine/Plugins/Experimental/GPULightmass/Source/GPULightmass/Private/VolumetricLightmap.cpp
+++ b/Engine/Plugins/Experimental/GPULightmass/Source/GPULightmass/Private/VolumetricLightmap.cpp
@@ -122,12 +122,12 @@ void FVolumetricLightmapRenderer::VoxelizeScene()
 
 	FRHICommandListExecutor::GetImmediateCommandList().ImmediateFlush(EImmediateFlushType::FlushRHIThreadFlushResources);
 
-	VolumeMin = CombinedImportanceVolume.Min;
+	VolumeMin = Scene->CombinedImportanceVolume.Min;
 
 	FIntVector FullGridSize(
-		FMath::TruncToInt(2 * CombinedImportanceVolume.GetExtent().X / TargetDetailCellSize) + 1,
-		FMath::TruncToInt(2 * CombinedImportanceVolume.GetExtent().Y / TargetDetailCellSize) + 1,
-		FMath::TruncToInt(2 * CombinedImportanceVolume.GetExtent().Z / TargetDetailCellSize) + 1);
+		FMath::TruncToInt(2 * Scene->CombinedImportanceVolume.GetExtent().X / TargetDetailCellSize) + 1,
+		FMath::TruncToInt(2 * Scene->CombinedImportanceVolume.GetExtent().Y / TargetDetailCellSize) + 1,
+		FMath::TruncToInt(2 * Scene->CombinedImportanceVolume.GetExtent().Z / TargetDetailCellSize) + 1);
 
 	const int32 BrickSizeLog2 = FMath::FloorLog2(BrickSize);
 	const int32 DetailCellsPerTopLevelBrick = 1 << (MaxRefinementLevels * BrickSizeLog2);
@@ -223,7 +223,7 @@ void FVolumetricLightmapRenderer::VoxelizeScene()
 				FComputeShaderUtils::GetGroupCount(VoxelizationVolumeMips[MipLevel]->GetDesc().GetSize(), FIntVector(4)));
 		}
 
-		for (FBox ImportanceVolume : ImportanceVolumes)
+		for (FBox ImportanceVolume : Scene->ImportanceVolumes)
 		{
 			TShaderMapRef<FVoxelizeImportanceVolumeCS> ComputeShader(GlobalShaderMap);
 
@@ -266,7 +266,7 @@ void FVolumetricLightmapRenderer::VoxelizeScene()
 			DrawDynamicMeshPass(
 				*Scene->ReferenceView,
 				RHICmdList,
-				[&Scene = Scene, View = Scene->ReferenceView.Get(), ImportanceVolumes = ImportanceVolumes](FDynamicPassMeshDrawListContext* DynamicMeshPassContext)
+				[&Scene = Scene, View = Scene->ReferenceView.Get(), ImportanceVolumes = Scene->ImportanceVolumes](FDynamicPassMeshDrawListContext* DynamicMeshPassContext)
 			{
 				FVLMVoxelizationMeshProcessor MeshProcessor(nullptr, View, DynamicMeshPassContext);
 
diff --git a/Engine/Plugins/Experimental/GPULightmass/Source/GPULightmass/Private/VolumetricLightmap.h b/Engine/Plugins/Experimental/GPULightmass/Source/GPULightmass/Private/VolumetricLightmap.h
--- a/Engine/Plugins/Experimental/GPULightmass/Source/GPULightmass/Private/VolumetricLightmap.h
+++ b/Engine/Plugins/Experimental/GPULightmass/Source/GPULightmass/Private/VolumetricLightmap.h
@@ -21,8 +21,6 @@ public:
 
 	FPrecomputedVolumetricLightmap* GetPrecomputedVolumetricLightmapForPreview();
 
-	FBox CombinedImportanceVolume;
-	TArray<FBox> ImportanceVolumes;
 	float TargetDetailCellSize = 50.0f;
 	int32 NumTotalBricks = 0;
 
