GPULM: Implement 'Aggressive Leak Prevention' option (backface detection inside irradiance cache) to prevent 1-texel wide leaks through thin walls around overlapping geometries

diff --git a/Engine/Plugins/Experimental/GPULightmass/Shaders/Private/IrradianceCacheVisualization.usf b/Engine/Plugins/Experimental/GPULightmass/Shaders/Private/IrradianceCacheVisualization.usf
--- a/Engine/Plugins/Experimental/GPULightmass/Shaders/Private/IrradianceCacheVisualization.usf
+++ b/Engine/Plugins/Experimental/GPULightmass/Shaders/Private/IrradianceCacheVisualization.usf
@@ -31,11 +31,22 @@ void VisualizeIrradianceCachePS(
 	{
 		uint RecordIndex = IrradianceCachingParameters.RWHashToIndex[Index];
 		uint4 RecordIrradianceAndSampleCount = IrradianceCachingParameters.IrradianceCacheRecords[RecordIndex];
+		uint BackfaceHitsCount = IrradianceCachingParameters.IrradianceCacheRecordBackfaceHits[RecordIndex].a;
 		if (RecordIrradianceAndSampleCount.w > 0)
 		{
 			bIrradianceQuerySuccessful = true;
 			
 			RecordIrradiance = asfloat(RecordIrradianceAndSampleCount.xyz) / RecordIrradianceAndSampleCount.w;
+
+			if (BackfaceHitsCount > RecordIrradianceAndSampleCount.w * BackfaceThresholdRejection)
+			{
+				RecordIrradiance = float3(1, 1, 0) / View.PreExposure;
+			}
+
+			if (BackfaceHitsCount > RecordIrradianceAndSampleCount.w * BackfaceThresholdInsideGeometry)
+			{
+				RecordIrradiance = float3(1, 0, 0) / View.PreExposure;
+			}
 		}
 		
 		bGeometryQuerySuccessful = true;
diff --git a/Engine/Plugins/Experimental/GPULightmass/Shaders/Private/IrradianceCachingCommon.ush b/Engine/Plugins/Experimental/GPULightmass/Shaders/Private/IrradianceCachingCommon.ush
--- a/Engine/Plugins/Experimental/GPULightmass/Shaders/Private/IrradianceCachingCommon.ush
+++ b/Engine/Plugins/Experimental/GPULightmass/Shaders/Private/IrradianceCachingCommon.ush
@@ -1,5 +1,12 @@
 #include "/Engine/Private/HashTable.ush"
 
+#ifndef IC_BACKFACE_DETECTION
+#define IC_BACKFACE_DETECTION 0
+#endif
+
+#define BackfaceThresholdRejection (1.0f / 64)
+#define BackfaceThresholdInsideGeometry (2.0f)
+
 struct FFinalGatherHitPoint
 {
 	float3 WorldPosition;
@@ -240,6 +247,7 @@ void EmitGeometryHitPoint(FFinalGatherHitPoint HitPoint, uint Size)
 		}
 		IrradianceCachingParameters.RWIndexToHash[RecordIndex] = Index;
 		IrradianceCachingParameters.IrradianceCacheRecords[RecordIndex] = uint4(0, 0, 0, 0);
+		IrradianceCachingParameters.IrradianceCacheRecordBackfaceHits[RecordIndex] = uint4(0, 0, 0, 0);
 	}
 }
 
diff --git a/Engine/Plugins/Experimental/GPULightmass/Shaders/Private/LightmapPathTracing.usf b/Engine/Plugins/Experimental/GPULightmass/Shaders/Private/LightmapPathTracing.usf
--- a/Engine/Plugins/Experimental/GPULightmass/Shaders/Private/LightmapPathTracing.usf
+++ b/Engine/Plugins/Experimental/GPULightmass/Shaders/Private/LightmapPathTracing.usf
@@ -445,6 +445,7 @@ void PathTracingKernel(
 				uint RecordIndex = IrradianceCachingParameters.RWHashToIndex[Index];
 
 				uint4 RecordIrradianceAndSampleCount = IrradianceCachingParameters.IrradianceCacheRecords[RecordIndex];
+				uint BackfaceHitsCount = IrradianceCachingParameters.IrradianceCacheRecordBackfaceHits[RecordIndex].a;
 
 				// Compute an overestimated upper bound for the number of samples we need for this texel by assuming
 				// all the rays from all IC building passes are sent from this single texel
@@ -465,6 +466,21 @@ void PathTracingKernel(
 				// Volumetric lightmaps
 				bShouldAcceptCacheEntry = RenderPassIndex >= GetIrradianceCachingQuality();
 			#endif
+				
+			#if IC_BACKFACE_DETECTION
+				// Reject the entry if it can see back faces
+				if (RenderPassIndex >= GetIrradianceCachingQuality() && BackfaceHitsCount > RecordIrradianceAndSampleCount.w * BackfaceThresholdRejection)
+				{
+					bShouldAcceptCacheEntry = false;
+				}
+
+				// Terminate the path if the entry contains too many back face hits which means it is very likely to be fully inside geometry
+				if (RenderPassIndex >= GetIrradianceCachingQuality() && BackfaceHitsCount > RecordIrradianceAndSampleCount.w * BackfaceThresholdInsideGeometry)
+				{
+					bIsValidSample = false;
+					break;
+				}
+			#endif
 			
 				if (bShouldAcceptCacheEntry)
 				{
@@ -836,6 +852,15 @@ void PathTracingKernel(
 			InterlockedAdd(IrradianceCachingParameters.IrradianceCacheRecords[NearestCacheEntryID].a, 1);
 		}
 	}
+	else
+	{
+	#if IC_BACKFACE_DETECTION
+		if(NearestCacheEntryID != -1)
+		{
+			InterlockedAdd(IrradianceCachingParameters.IrradianceCacheRecordBackfaceHits[NearestCacheEntryID].a, 1);
+		}
+	#endif
+	}
 #endif
 	
 	RadianceValue = Radiance;
diff --git a/Engine/Plugins/Experimental/GPULightmass/Source/GPULightmass/Private/IrradianceCaching.cpp b/Engine/Plugins/Experimental/GPULightmass/Source/GPULightmass/Private/IrradianceCaching.cpp
--- a/Engine/Plugins/Experimental/GPULightmass/Source/GPULightmass/Private/IrradianceCaching.cpp
+++ b/Engine/Plugins/Experimental/GPULightmass/Source/GPULightmass/Private/IrradianceCaching.cpp
@@ -21,10 +21,14 @@ FIrradianceCache::FIrradianceCache(int32 Quality, float Spacing, float CornerRej
 		FRHIResourceCreateInfo CreateInfo(TEXT("FIrradianceCache"));
 		IrradianceCacheRecords = RHICreateStructuredBuffer(sizeof(FVector4f), sizeof(FVector4f) * IrradianceCacheMaxSize, BUF_UnorderedAccess | BUF_ShaderResource, CreateInfo);
 		IrradianceCacheRecordsUAV = RHICreateUnorderedAccessView(IrradianceCacheRecords, false, false);
+		IrradianceCacheRecordBackfaceHits = RHICreateStructuredBuffer(sizeof(FVector4f), sizeof(FVector4f) * IrradianceCacheMaxSize, BUF_UnorderedAccess | BUF_ShaderResource, CreateInfo);
+		IrradianceCacheRecordBackfaceHitsUAV = RHICreateUnorderedAccessView(IrradianceCacheRecordBackfaceHits, false, false);
 		
+		IrradianceCacheTotalBytes += sizeof(FVector4f) * IrradianceCacheMaxSize;
 		IrradianceCacheTotalBytes += sizeof(FVector4f) * IrradianceCacheMaxSize;
 
 		RHICmdList.ClearUAVUint(IrradianceCacheRecordsUAV, FUintVector4(0, 0, 0, 0));
+		RHICmdList.ClearUAVUint(IrradianceCacheRecordBackfaceHitsUAV, FUintVector4(0, 0, 0, 0));
 	}
 
 	int32 HashTableSize = IrradianceCacheMaxSize * 4;
@@ -58,6 +62,7 @@ FIrradianceCache::FIrradianceCache(int32 Quality, float Spacing, float CornerRej
 
 	FIrradianceCachingParameters IrradianceCachingParameters;
 	IrradianceCachingParameters.IrradianceCacheRecords = IrradianceCacheRecordsUAV;
+	IrradianceCachingParameters.IrradianceCacheRecordBackfaceHits = IrradianceCacheRecordBackfaceHitsUAV;
 	IrradianceCachingParameters.Quality = Quality;
 	IrradianceCachingParameters.Spacing = Spacing;
 	IrradianceCachingParameters.CornerRejection = CornerRejection;
diff --git a/Engine/Plugins/Experimental/GPULightmass/Source/GPULightmass/Private/IrradianceCaching.h b/Engine/Plugins/Experimental/GPULightmass/Source/GPULightmass/Private/IrradianceCaching.h
--- a/Engine/Plugins/Experimental/GPULightmass/Source/GPULightmass/Private/IrradianceCaching.h
+++ b/Engine/Plugins/Experimental/GPULightmass/Source/GPULightmass/Private/IrradianceCaching.h
@@ -13,6 +13,7 @@ BEGIN_UNIFORM_BUFFER_STRUCT(FIrradianceCachingParameters, )
 	SHADER_PARAMETER(float, Spacing)
 	SHADER_PARAMETER(float, CornerRejection)
 	SHADER_PARAMETER_UAV(RWStructuredBuffer<uint4>, IrradianceCacheRecords)
+	SHADER_PARAMETER_UAV(RWStructuredBuffer<uint4>, IrradianceCacheRecordBackfaceHits)
 	SHADER_PARAMETER_UAV(RWStructuredBuffer<uint2>, RWHashTable)
 	SHADER_PARAMETER_UAV(RWStructuredBuffer<uint>, RWHashToIndex)
 	SHADER_PARAMETER_UAV(RWStructuredBuffer<uint>, RWIndexToHash)
@@ -26,6 +27,8 @@ struct FIrradianceCache
 
 	FBufferRHIRef IrradianceCacheRecords;
 	FUnorderedAccessViewRHIRef IrradianceCacheRecordsUAV;
+	FBufferRHIRef IrradianceCacheRecordBackfaceHits;
+	FUnorderedAccessViewRHIRef IrradianceCacheRecordBackfaceHitsUAV;
 
 	TUniformBufferRef<FIrradianceCachingParameters> IrradianceCachingParametersUniformBuffer;
 
diff --git a/Engine/Plugins/Experimental/GPULightmass/Source/GPULightmass/Private/LightmapRayTracing.h b/Engine/Plugins/Experimental/GPULightmass/Source/GPULightmass/Private/LightmapRayTracing.h
--- a/Engine/Plugins/Experimental/GPULightmass/Source/GPULightmass/Private/LightmapRayTracing.h
+++ b/Engine/Plugins/Experimental/GPULightmass/Source/GPULightmass/Private/LightmapRayTracing.h
@@ -41,11 +41,19 @@ class FLightmapPathTracingRGS : public FGlobalShader
 
 	class FUseIrradianceCaching : SHADER_PERMUTATION_BOOL("USE_IRRADIANCE_CACHING");
 	class FUseFirstBounceRayGuiding : SHADER_PERMUTATION_BOOL("USE_FIRST_BOUNCE_RAY_GUIDING");
+	class FUseICBackfaceDetection : SHADER_PERMUTATION_BOOL("IC_BACKFACE_DETECTION");
 
-	using FPermutationDomain = TShaderPermutationDomain<FUseIrradianceCaching, FUseFirstBounceRayGuiding>;
+	using FPermutationDomain = TShaderPermutationDomain<FUseIrradianceCaching, FUseFirstBounceRayGuiding, FUseICBackfaceDetection>;
 
 	static bool ShouldCompilePermutation(const FGlobalShaderPermutationParameters& Parameters)
 	{
+		FPermutationDomain PermutationVector(Parameters.PermutationId);
+		// Disallow invalid permutations
+		if (!PermutationVector.Get<FUseIrradianceCaching>())
+		{
+			if (PermutationVector.Get<FUseICBackfaceDetection>()) return false;
+			if (PermutationVector.Get<FUseFirstBounceRayGuiding>()) return false;
+		}
 		return EnumHasAllFlags(Parameters.Flags, EShaderPermutationFlags::HasEditorOnlyData) && ShouldCompileRayTracingShadersForProject(Parameters.Platform);
 	}
 
diff --git a/Engine/Plugins/Experimental/GPULightmass/Source/GPULightmass/Private/LightmapRenderer.cpp b/Engine/Plugins/Experimental/GPULightmass/Source/GPULightmass/Private/LightmapRenderer.cpp
--- a/Engine/Plugins/Experimental/GPULightmass/Source/GPULightmass/Private/LightmapRenderer.cpp
+++ b/Engine/Plugins/Experimental/GPULightmass/Source/GPULightmass/Private/LightmapRenderer.cpp
@@ -1217,8 +1217,9 @@ bool FSceneRenderState::SetupRayTracingScene(int32 LODIndex)
 			TArray<FRHIRayTracingShader*> RayGenShaderTable;
 			{
 				FLightmapPathTracingRGS::FPermutationDomain PermutationVector;
-				PermutationVector.Set<FLightmapPathTracingRGS::FUseFirstBounceRayGuiding>(Settings->bUseFirstBounceRayGuiding);
+				PermutationVector.Set<FLightmapPathTracingRGS::FUseFirstBounceRayGuiding>(Settings->bUseIrradianceCaching && Settings->bUseFirstBounceRayGuiding);
 				PermutationVector.Set<FLightmapPathTracingRGS::FUseIrradianceCaching>(Settings->bUseIrradianceCaching);
+				PermutationVector.Set<FLightmapPathTracingRGS::FUseICBackfaceDetection>(Settings->bUseIrradianceCaching && Settings->bUseIrradianceCacheBackfaceDetection);
 				RayGenShaderTable.Add(GlobalShaderMap->GetShader<FLightmapPathTracingRGS>(PermutationVector).GetRayTracingShader());
 			}
 			{
@@ -2383,8 +2384,9 @@ void FLightmapRenderer::Finalize(FRDGBuilder& GraphBuilder)
 
 
 									FLightmapPathTracingRGS::FPermutationDomain PermutationVector;
-									PermutationVector.Set<FLightmapPathTracingRGS::FUseFirstBounceRayGuiding>(Scene->Settings->bUseFirstBounceRayGuiding);
+									PermutationVector.Set<FLightmapPathTracingRGS::FUseFirstBounceRayGuiding>(Scene->Settings->bUseIrradianceCaching && Scene->Settings->bUseFirstBounceRayGuiding);
 									PermutationVector.Set<FLightmapPathTracingRGS::FUseIrradianceCaching>(Scene->Settings->bUseIrradianceCaching);
+									PermutationVector.Set<FLightmapPathTracingRGS::FUseICBackfaceDetection>(Scene->Settings->bUseIrradianceCaching && Scene->Settings->bUseIrradianceCacheBackfaceDetection);
 									auto RayGenerationShader = GlobalShaderMap->GetShader<FLightmapPathTracingRGS>(PermutationVector);
 									ClearUnusedGraphResources(RayGenerationShader, PassParameters);
 
diff --git a/Engine/Plugins/Experimental/GPULightmass/Source/GPULightmass/Public/GPULightmassSettings.h b/Engine/Plugins/Experimental/GPULightmass/Source/GPULightmass/Public/GPULightmassSettings.h
--- a/Engine/Plugins/Experimental/GPULightmass/Source/GPULightmass/Public/GPULightmassSettings.h
+++ b/Engine/Plugins/Experimental/GPULightmass/Source/GPULightmass/Public/GPULightmassSettings.h
@@ -53,7 +53,7 @@ public:
 	UPROPERTY(EditAnywhere, BlueprintReadWrite, Category = Denoising, DisplayName = "Denoise when")
 	EGPULightmassDenoisingOptions DenoisingOptions = EGPULightmassDenoisingOptions::OnCompletion;
 	
-	UPROPERTY(EditAnywhere, BlueprintReadWrite, Category = Denoising)
+	UPROPERTY(EditAnywhere, BlueprintReadWrite, Category = Denoising, meta = (EditCondition = "DenoisingOptions != EGPULightmassDenoisingOptions::None"))
 	EGPULightmassDenoiser Denoiser = EGPULightmassDenoiser::IntelOIDN;
 
 	// Whether to compress lightmap textures.  Disabling lightmap texture compression will reduce artifacts but increase memory and disk size by 4x.
@@ -94,6 +94,10 @@ public:
 	// Number of samples per Irradiance Cache cell.
 	UPROPERTY(EditAnywhere, BlueprintReadWrite, Category = IrradianceCaching, DisplayName = "Quality", meta = (EditCondition = "bUseIrradianceCaching", ClampMin = "4", ClampMax = "65536", UIMax = "8192"))
 	int32 IrradianceCacheQuality = 128;
+	
+	// Further prevent leaks caused by irradiance cache cells being placed inside geometry, at the cost of more fireflies and slower sampling speed. Recommended to be used with higher irradiance cache quality (>=256)
+	UPROPERTY(EditAnywhere, BlueprintReadWrite, AdvancedDisplay, Category = IrradianceCaching, DisplayName = "Aggressive Leak Prevention", meta = (EditCondition = "bUseIrradianceCaching"))
+	bool bUseIrradianceCacheBackfaceDetection = false;
 
 	// Size of each Irradiance Cache cell. Smaller sizes will be slower but more accurate.
 	UPROPERTY(EditAnywhere, BlueprintReadWrite, AdvancedDisplay, Category = IrradianceCaching, DisplayName = "Size", meta = (EditCondition = "bUseIrradianceCaching", ClampMin = "4", ClampMax = "1024"))
