Fix non-unity, nopch build

diff --git a/Engine/Plugins/Experimental/GPULightmass/Source/GPULightmass/Private/LightmapGBuffer.h b/Engine/Plugins/Experimental/GPULightmass/Source/GPULightmass/Private/LightmapGBuffer.h
--- a/Engine/Plugins/Experimental/GPULightmass/Source/GPULightmass/Private/LightmapGBuffer.h
+++ b/Engine/Plugins/Experimental/GPULightmass/Source/GPULightmass/Private/LightmapGBuffer.h
@@ -8,6 +8,7 @@
 #include "MeshPassProcessor.inl"
 #include "MeshMaterialShader.h"
 #include "LightMapRendering.h"
+#include "Materials/Material.h"
 
 BEGIN_GLOBAL_SHADER_PARAMETER_STRUCT(FLightmapGBufferParams, )
 	SHADER_PARAMETER_RDG_TEXTURE_UAV(RWTexture2D<float4>, ScratchTilePoolLayer0)
