GPULM: Fix SkyLight invalidation issues

diff --git a/Engine/Plugins/Experimental/GPULightmass/Source/GPULightmass/Private/GPULightmass.cpp b/Engine/Plugins/Experimental/GPULightmass/Source/GPULightmass/Private/GPULightmass.cpp
--- a/Engine/Plugins/Experimental/GPULightmass/Source/GPULightmass/Private/GPULightmass.cpp
+++ b/Engine/Plugins/Experimental/GPULightmass/Source/GPULightmass/Private/GPULightmass.cpp
@@ -170,7 +170,6 @@ void FGPULightmass::OnLightComponentRegistered(ULightComponentBase* InComponent)
 
 	if (UDirectionalLightComponent* DirectionalLight = Cast<UDirectionalLightComponent>(InComponent))
 	{
-		Scene.ConditionalTriggerSkyLightRecapture();
 		Scene.AddLight(DirectionalLight);
 	}
 	else if (URectLightComponent* RectLight = Cast<URectLightComponent>(InComponent))
@@ -198,7 +197,6 @@ void FGPULightmass::OnLightComponentUnregistered(ULightComponentBase* InComponen
 	if (UDirectionalLightComponent* DirectionalLight = Cast<UDirectionalLightComponent>(InComponent))
 	{
 		Scene.RemoveLight(DirectionalLight);
-		Scene.ConditionalTriggerSkyLightRecapture();
 	}
 	else if (URectLightComponent* RectLight = Cast<URectLightComponent>(InComponent))
 	{
diff --git a/Engine/Plugins/Experimental/GPULightmass/Source/GPULightmass/Private/GPULightmassSettings.cpp b/Engine/Plugins/Experimental/GPULightmass/Source/GPULightmass/Private/GPULightmassSettings.cpp
--- a/Engine/Plugins/Experimental/GPULightmass/Source/GPULightmass/Private/GPULightmassSettings.cpp
+++ b/Engine/Plugins/Experimental/GPULightmass/Source/GPULightmass/Private/GPULightmassSettings.cpp
@@ -152,6 +152,8 @@ void UGPULightmassSubsystem::Launch()
 		SlowTask.MakeDialog();
 		SlowTask.EnterProgressFrame(1, LOCTEXT("StartingStaticLightingSystem", "Starting static lighting system"));
 
+		World->UpdateAllSkyCaptures();
+		
 		{
 			FGlobalComponentRecreateRenderStateContext RecreateRenderStateContext; // Implicit FlushRenderingCommands();
 
diff --git a/Engine/Plugins/Experimental/GPULightmass/Source/GPULightmass/Private/Scene/Scene.cpp b/Engine/Plugins/Experimental/GPULightmass/Source/GPULightmass/Private/Scene/Scene.cpp
--- a/Engine/Plugins/Experimental/GPULightmass/Source/GPULightmass/Private/Scene/Scene.cpp
+++ b/Engine/Plugins/Experimental/GPULightmass/Source/GPULightmass/Private/Scene/Scene.cpp
@@ -663,14 +663,6 @@ void FScene::RemoveLight(USkyLightComponent* SkyLight)
 
 void FScene::OnSkyAtmosphereModified()
 {
-	TArray<UDirectionalLightComponent*> PreviouslyRegisteredDirectionalLight;
-	LightScene.RegisteredDirectionalLightComponentUObjects.GenerateKeyArray(PreviouslyRegisteredDirectionalLight);	
-	for(UDirectionalLightComponent* DirectionalLight : PreviouslyRegisteredDirectionalLight)
-	{
-		RemoveLight(DirectionalLight);
-		AddLight(DirectionalLight);
-	}
-
 	ConditionalTriggerSkyLightRecapture();
 }
 
@@ -681,10 +673,8 @@ void FScene::ConditionalTriggerSkyLightRecapture()
 		USkyLightComponent* SkyLight = LightScene.SkyLight->ComponentUObject;
 		if (SkyLight->SourceType == SLS_CapturedScene || SkyLight->bRealTimeCapture)
 		{
-			RemoveLight(SkyLight);
 			SkyLight->SetCaptureIsDirty();
 			SkyLight->UpdateSkyCaptureContents(SkyLight->GetWorld());
-			AddLight(SkyLight);
 		}
 	}
 }
