GPULM: Fix landscape crash

diff --git a/Engine/Plugins/Experimental/GPULightmass/Source/GPULightmass/Private/Scene/Scene.cpp b/Engine/Plugins/Experimental/GPULightmass/Source/GPULightmass/Private/Scene/Scene.cpp
--- a/Engine/Plugins/Experimental/GPULightmass/Source/GPULightmass/Private/Scene/Scene.cpp
+++ b/Engine/Plugins/Experimental/GPULightmass/Source/GPULightmass/Private/Scene/Scene.cpp
@@ -1342,6 +1342,9 @@ void FScene::AddGeometryInstanceFromComponent(ULandscapeComponent* InComponent)
 			LandscapeParams.HeightmapTexture = Initializer.HeightmapTexture->TextureReference.TextureReferenceRHI;
 			LandscapeParams.HeightmapTextureSampler = TStaticSamplerState<SF_Point>::GetRHI();
 
+			LandscapeParams.NormalmapTexture = Initializer.HeightmapTexture->TextureReference.TextureReferenceRHI;
+			LandscapeParams.NormalmapTextureSampler = TStaticSamplerState<SF_Point>::GetRHI();
+			
 			// No support for XYOffset
 			LandscapeParams.XYOffsetmapTexture = GBlackTexture->TextureRHI;
 			LandscapeParams.XYOffsetmapTextureSampler = GBlackTexture->SamplerStateRHI;
