GPULM: Change how VLM sample counting works to resolve multiple related issues

* Now it takes 2 samples (upper and lower hemisphere) per SampleCount

diff --git a/Engine/Plugins/Experimental/GPULightmass/Shaders/Private/BrickAllocationManagement.usf b/Engine/Plugins/Experimental/GPULightmass/Shaders/Private/BrickAllocationManagement.usf
--- a/Engine/Plugins/Experimental/GPULightmass/Shaders/Private/BrickAllocationManagement.usf
+++ b/Engine/Plugins/Experimental/GPULightmass/Shaders/Private/BrickAllocationManagement.usf
@@ -232,7 +232,7 @@ UNORM float4 DilateWithValidityMask(RWTexture3D<UNORM float4> Texture, Texture3D
 	return Value / max(AccumulatedWeight, 0.00001f);
 }
 
-UNORM float4 DilateWithValidityMask(RWTexture3D<UNORM float> Texture, Texture3D<float4> CoverageMask, float3 UVW)
+UNORM float DilateWithValidityMask(RWTexture3D<UNORM float> Texture, Texture3D<float4> CoverageMask, float3 UVW)
 {
 	int3 Coord = UVW;
 	float3 CoordinateFraction = frac(UVW);
@@ -374,7 +374,7 @@ void StitchBorderCS(uint3 BrickID : SV_GroupID, uint3 CellPosInBrick : SV_GroupT
 		OutSkyBentNormal[VoxelPos] = DilateWithValidityMask(OutSkyBentNormal, AmbientVector, IndirectionTextureUVW);
 		OutDirectionalLightShadowing[VoxelPos] = DilateWithValidityMask(OutDirectionalLightShadowing, AmbientVector, IndirectionTextureUVW);
 	}
-	
+
 	for (int dx = -1; dx <= 1; dx++)
 	{
 		for (int dy = -1; dy <= 1; dy++)
diff --git a/Engine/Plugins/Experimental/GPULightmass/Shaders/Private/LightmapPathTracing.usf b/Engine/Plugins/Experimental/GPULightmass/Shaders/Private/LightmapPathTracing.usf
--- a/Engine/Plugins/Experimental/GPULightmass/Shaders/Private/LightmapPathTracing.usf
+++ b/Engine/Plugins/Experimental/GPULightmass/Shaders/Private/LightmapPathTracing.usf
@@ -1222,12 +1222,11 @@ void VolumetricLightmapPathTracingMainRG()
 	if (BrickIndex >= NumTotalBricks) return;
 	uint CellIndex = DispatchRaysIndex().x % BrickVolumeSize;
 	uint3 CellPosInBrick = ComputeBrickLayoutPosition(CellIndex, uint3(5, 5, 5));
+	uint PaddedBrickSize = 4 + 1;
+	uint3 VoxelPos = ComputeBrickLayoutPosition(BrickIndex, uint3(256, 256, 256)) * PaddedBrickSize + CellPosInBrick;
 	
 	bool bIsValidSample = true;	
 	
-	int3 CellPosInVLM = (BrickRequests[BrickIndex].xyz * 4 + CellPosInBrick) * BrickRequests[BrickIndex].w;	
-	uint Seed = ComputeBrickLinearAddress(CellPosInVLM, IndirectionTextureDim * 4);
-	
 	int EffectiveRenderPassIndex = FrameNumber;
 #if USE_IRRADIANCE_CACHING
 	if (EffectiveRenderPassIndex >= GetIrradianceCachingQuality())
@@ -1235,80 +1234,6 @@ void VolumetricLightmapPathTracingMainRG()
 		EffectiveRenderPassIndex -= GetIrradianceCachingQuality();
 	}
 #endif
-
-	#if 0
-		float3 Jitter = RandomSequence_GenerateSample3D(RandSequence);
-	#else
-		float3 RandSample = float3(
-			Halton(MortonEncode3(CellPosInVLM) + EffectiveRenderPassIndex / 2, 2), 
-			Halton(MortonEncode3(CellPosInVLM) + EffectiveRenderPassIndex / 2, 3),
-			Halton(MortonEncode3(CellPosInVLM) + EffectiveRenderPassIndex / 2, 5)
-		);
-		float3 Jitter = RandSample;
-	#endif		
-	
-	float3 JitteredCellPosInVLM = (BrickRequests[BrickIndex].xyz * 4 + CellPosInBrick + Jitter - float3(0.5, 0.5, 0.5)) * BrickRequests[BrickIndex].w;	
-	float3 DetailCellSize = VolumeSize.xyz / IndirectionTextureDim / 4;
-	float3 WorldPosition = VolumeMin.xyz + DetailCellSize * JitteredCellPosInVLM; // RT_LWC_TODO
-	float3 TranslatedWorldPosition = WorldPosition + LWCHackToFloat(PrimaryView.PreViewTranslation); // RT_LWC_TODO
-	float3 ShadingNormal = float3(0, 0, (EffectiveRenderPassIndex % 2 == 0) ? 1 : -1);
-
-	float3 RadianceValue = 0;
-	float3 RadianceDirection = 0;				
-	float3 DirectLightingIrradiance = 0;
-#ifdef LIGHTMAP_PATH_TRACING_MAIN_RG
-	FL2SHAndCorrection DirectLightingSH = (FL2SHAndCorrection)0;
-#else
-	FThreeBandSHVectorRGBFloat DirectLightingSH = (FThreeBandSHVectorRGBFloat)0;
-#endif
-	float LuminanceForFirstBounceRayGuiding = 0;
-	float3 SkyLightBentNormal = 0;
-	float4 PrimaryRandSample;
-
-	if (bIsValidSample)
-	{
-		RandomSequence RandSequence;
-		RandomSequence_Initialize(RandSequence, Seed, EffectiveRenderPassIndex / 2);	
-		uint SampleIndex = 4;
-		
-		PathTracingKernel(
-			FrameNumber,
-			TranslatedWorldPosition,
-			ShadingNormal,
-			RandSequence,
-			SampleIndex,
-			bIsValidSample,
-			RadianceValue,
-			RadianceDirection,
-			DirectLightingIrradiance,
-			DirectLightingSH,
-			LuminanceForFirstBounceRayGuiding,
-			SkyLightBentNormal,
-			PrimaryRandSample);
-	}
-	
-	// Every even sample is upper hemisphere, odd sample is lower hemisphere
-	// SampleCount stored is twice as the actual full sphere sample count
-	RadianceValue *= 2;
-#ifndef LIGHTMAP_PATH_TRACING_MAIN_RG
-	DirectLightingSH.R.V0 *= 2;
-	DirectLightingSH.R.V1 *= 2;
-	DirectLightingSH.R.V2 *= 2;
-	DirectLightingSH.G.V0 *= 2;
-	DirectLightingSH.G.V1 *= 2;
-	DirectLightingSH.G.V2 *= 2;
-	DirectLightingSH.B.V0 *= 2;
-	DirectLightingSH.B.V1 *= 2;
-	DirectLightingSH.B.V2 *= 2;
-#endif
-
-	if (any(isnan(RadianceValue)) || any(RadianceValue < 0))
-	{
-		bIsValidSample = false;
-	}
-
-	uint PaddedBrickSize = 4 + 1;
-	uint3 VoxelPos = ComputeBrickLayoutPosition(BrickIndex, uint3(256, 256, 256)) * PaddedBrickSize + CellPosInBrick;
 	
 	if (EffectiveRenderPassIndex == 0) 
 	{
@@ -1326,11 +1251,70 @@ void VolumetricLightmapPathTracingMainRG()
 		SkyBentNormal[VoxelPos] = float4(0, 0, 0, 0);
 	}
 	
-	if (bIsValidSample)
+	int3 CellPosInVLM = (BrickRequests[BrickIndex].xyz * 4 + CellPosInBrick) * BrickRequests[BrickIndex].w;	
+	uint Seed = ComputeBrickLinearAddress(CellPosInVLM, IndirectionTextureDim * 4);
+	
+	float3 RandSample = float3(
+		Halton(MortonEncode3(CellPosInVLM) + EffectiveRenderPassIndex, 2), 
+		Halton(MortonEncode3(CellPosInVLM) + EffectiveRenderPassIndex, 3),
+		Halton(MortonEncode3(CellPosInVLM) + EffectiveRenderPassIndex, 5)
+	);
+	float3 Jitter = RandSample;
+	
+	float3 JitteredCellPosInVLM = (BrickRequests[BrickIndex].xyz * 4 + CellPosInBrick + Jitter - float3(0.5, 0.5, 0.5)) * BrickRequests[BrickIndex].w;	
+	float3 DetailCellSize = VolumeSize.xyz / IndirectionTextureDim / 4;
+	float3 WorldPosition = VolumeMin.xyz + DetailCellSize * JitteredCellPosInVLM; // RT_LWC_TODO
+	float3 TranslatedWorldPosition = WorldPosition + LWCHackToFloat(PrimaryView.PreViewTranslation); // RT_LWC_TODO
+
+	int NumValidHalfSamples = 0;
+	
+	for (int HalfSampleIndex = 0; HalfSampleIndex < 2; HalfSampleIndex++)
 	{
-		uint SampleCount = asuint(AmbientVector[VoxelPos].w);
+		float3 ShadingNormal = float3(0, 0, (HalfSampleIndex == 0) ? 1 : -1);
+
+		float3 RadianceValue = 0;
+		float3 RadianceDirection = 0;				
+		float3 DirectLightingIrradiance = 0;
+#ifdef LIGHTMAP_PATH_TRACING_MAIN_RG
+		FL2SHAndCorrection DirectLightingSH = (FL2SHAndCorrection)0;
+#else
+		FThreeBandSHVectorRGBFloat DirectLightingSH = (FThreeBandSHVectorRGBFloat)0;
+#endif
+		float LuminanceForFirstBounceRayGuiding = 0;
+		float3 SkyLightBentNormal = 0;
+		float4 PrimaryRandSample;
+
+		if (bIsValidSample)
+		{
+			RandomSequence RandSequence;
+			RandomSequence_Initialize(RandSequence, Seed, EffectiveRenderPassIndex);	
+			uint SampleIndex = 4;
 		
+			PathTracingKernel(
+				FrameNumber,
+				TranslatedWorldPosition,
+				ShadingNormal,
+				RandSequence,
+				SampleIndex,
+				bIsValidSample,
+				RadianceValue,
+				RadianceDirection,
+				DirectLightingIrradiance,
+				DirectLightingSH,
+				LuminanceForFirstBounceRayGuiding,
+				SkyLightBentNormal,
+				PrimaryRandSample);
+		}
+
+		if (any(isnan(RadianceValue)) || any(RadianceValue < 0))
+		{
+			bIsValidSample = false;
+		}
+
+		if (bIsValidSample)
 		{
+			NumValidHalfSamples++;
+
 			FThreeBandSHVectorRGBFloat SH;
 			SH.R.V0 = SHBasisFunction3Float(RadianceDirection).V0 * RadianceValue.r;
 			SH.R.V1 = SHBasisFunction3Float(RadianceDirection).V1 * RadianceValue.r;
@@ -1341,7 +1325,7 @@ void VolumetricLightmapPathTracingMainRG()
 			SH.B.V0 = SHBasisFunction3Float(RadianceDirection).V0 * RadianceValue.b;
 			SH.B.V1 = SHBasisFunction3Float(RadianceDirection).V1 * RadianceValue.b;
 			SH.B.V2 = SHBasisFunction3Float(RadianceDirection).V2 * RadianceValue.b;
-			
+
 			AmbientVector[VoxelPos].x += SH.R.V0.x;
 			AmbientVector[VoxelPos].y += SH.G.V0.x;
 			AmbientVector[VoxelPos].z += SH.B.V0.x;
@@ -1354,32 +1338,32 @@ void VolumetricLightmapPathTracingMainRG()
 			
 			SHCoefficients0B[VoxelPos] += float4(SH.B.V0.yzw, SH.B.V1.x);
 			SHCoefficients1B[VoxelPos] += float4(SH.B.V1.yzw, SH.B.V2);
-			
-			if (EffectiveRenderPassIndex % 2 == 0)
+				
+			if (HalfSampleIndex == 0)
 			{
-				// Every even sample is upper hemisphere, odd sample is lower hemisphere
-				// SampleCount stored is twice as the actual full sphere sample count
-				SkyBentNormal[VoxelPos].rgb += SkyLightBentNormal * 2;
+				SkyBentNormal[VoxelPos].rgb = SkyLightBentNormal;
 			}
-		}
-
+		
 		#ifndef LIGHTMAP_PATH_TRACING_MAIN_RG
-		{
 			AmbientVector[VoxelPos].x += DirectLightingSH.R.V0.x;
 			AmbientVector[VoxelPos].y += DirectLightingSH.G.V0.x;
 			AmbientVector[VoxelPos].z += DirectLightingSH.B.V0.x;
-			
+				
 			SHCoefficients0R[VoxelPos] += float4(DirectLightingSH.R.V0.yzw, DirectLightingSH.R.V1.x);
 			SHCoefficients1R[VoxelPos] += float4(DirectLightingSH.R.V1.yzw, DirectLightingSH.R.V2);
-			
+				
 			SHCoefficients0G[VoxelPos] += float4(DirectLightingSH.G.V0.yzw, DirectLightingSH.G.V1.x);
 			SHCoefficients1G[VoxelPos] += float4(DirectLightingSH.G.V1.yzw, DirectLightingSH.G.V2);		
-			
+				
 			SHCoefficients0B[VoxelPos] += float4(DirectLightingSH.B.V0.yzw, DirectLightingSH.B.V1.x);
 			SHCoefficients1B[VoxelPos] += float4(DirectLightingSH.B.V1.yzw, DirectLightingSH.B.V2);
-		}
 		#endif
-		
+		}
+	}
+
+	if (NumValidHalfSamples > 0)
+	{
+		uint SampleCount = asuint(AmbientVector[VoxelPos].w);
 		AmbientVector[VoxelPos].w = asfloat(SampleCount + 1);
 	}
 	
@@ -1403,14 +1387,12 @@ void VolumetricLightmapPathTracingMainRG()
 		bool bIsValidRay = GenerateOcclusionRay(
 			LIGHT_TYPE_DIRECTIONAL,
 			LightParameters,
-			TranslatedWorldPosition, ShadingNormal,
+			TranslatedWorldPosition, float3(0, 0, 1),
 			RandSample,
 			/* out */ Ray.Origin,
 			/* out */ Ray.Direction,
 			/* out */ Ray.TMin,
 			/* out */ Ray.TMax);
-		
-		ApplyRayBias(Ray, 1.0f, ShadingNormal);
 			
 		if (bIsValidRay)
 		{
diff --git a/Engine/Plugins/Experimental/GPULightmass/Source/GPULightmass/Private/VolumetricLightmap.cpp b/Engine/Plugins/Experimental/GPULightmass/Source/GPULightmass/Private/VolumetricLightmap.cpp
--- a/Engine/Plugins/Experimental/GPULightmass/Source/GPULightmass/Private/VolumetricLightmap.cpp
+++ b/Engine/Plugins/Experimental/GPULightmass/Source/GPULightmass/Private/VolumetricLightmap.cpp
@@ -735,7 +735,18 @@ void FVolumetricLightmapRenderer::BackgroundTick()
 					FStitchBorderCS::FParameters* PassParameters = GraphBuilder.AllocParameters<FStitchBorderCS::FParameters>();
 					PassParameters->BrickDataDimensions = VolumetricLightmapData.BrickDataDimensions;
 					PassParameters->IndirectionTextureDim = IndirectionTextureDimensions;
-					PassParameters->FrameNumber = FrameNumber / NumFramesOneRound;
+
+					int32 EffectiveRenderPassIndex = FrameNumber / NumFramesOneRound;
+
+					if (Scene->Settings->bUseIrradianceCaching)
+					{
+						if (EffectiveRenderPassIndex >= Scene->Settings->IrradianceCacheQuality * Scene->Settings->VolumetricLightmapQualityMultiplier)
+						{
+							EffectiveRenderPassIndex -= Scene->Settings->IrradianceCacheQuality * Scene->Settings->VolumetricLightmapQualityMultiplier;
+						}
+					}
+					
+					PassParameters->FrameNumber = EffectiveRenderPassIndex;
 					PassParameters->NumTotalBricks = NumTotalBricks;
 					PassParameters->BrickBatchOffset = BrickBatchOffset;
 					PassParameters->IndirectionTexture = IndirectionTextureUAV;
