GPULM: VolumetricLightmapQualityMultiplier is no longer applied to IC quality as it is excessive and our IC passes are quite heavy now

diff --git a/Engine/Plugins/Experimental/GPULightmass/Shaders/Private/IrradianceCachingCommon.ush b/Engine/Plugins/Experimental/GPULightmass/Shaders/Private/IrradianceCachingCommon.ush
--- a/Engine/Plugins/Experimental/GPULightmass/Shaders/Private/IrradianceCachingCommon.ush
+++ b/Engine/Plugins/Experimental/GPULightmass/Shaders/Private/IrradianceCachingCommon.ush
@@ -253,9 +253,5 @@ void EmitGeometryHitPoint(FFinalGatherHitPoint HitPoint, uint Size)
 
 int GetIrradianceCachingQuality()
 {
-#ifdef VOLUMETRIC_LIGHTMAP_PATH_TRACING_MAIN_RG
-	return IrradianceCachingParameters.Quality * VolumetricLightmapQualityMultiplier;
-#endif
-	
 	return IrradianceCachingParameters.Quality;
 }
diff --git a/Engine/Plugins/Experimental/GPULightmass/Shaders/Private/LightmapPathTracing.usf b/Engine/Plugins/Experimental/GPULightmass/Shaders/Private/LightmapPathTracing.usf
--- a/Engine/Plugins/Experimental/GPULightmass/Shaders/Private/LightmapPathTracing.usf
+++ b/Engine/Plugins/Experimental/GPULightmass/Shaders/Private/LightmapPathTracing.usf
@@ -30,10 +30,6 @@ RWTexture2D<uint> RayGuidingLuminance;
 Texture2D<float> RayGuidingCDFX;
 Texture2D<float> RayGuidingCDFY;
 
-#ifdef VOLUMETRIC_LIGHTMAP_PATH_TRACING_MAIN_RG
-int VolumetricLightmapQualityMultiplier;
-#endif
-
 #include "IrradianceCachingCommon.ush"
 
 #include "/Engine/Private/ShadingModels.ush"
diff --git a/Engine/Plugins/Experimental/GPULightmass/Source/GPULightmass/Private/LightmapRayTracing.h b/Engine/Plugins/Experimental/GPULightmass/Source/GPULightmass/Private/LightmapRayTracing.h
--- a/Engine/Plugins/Experimental/GPULightmass/Source/GPULightmass/Private/LightmapRayTracing.h
+++ b/Engine/Plugins/Experimental/GPULightmass/Source/GPULightmass/Private/LightmapRayTracing.h
@@ -132,7 +132,6 @@ class FVolumetricLightmapPathTracingRGS : public FGlobalShader
 		SHADER_PARAMETER(FIntVector, IndirectionTextureDim)
 		SHADER_PARAMETER(int32, NumTotalBricks)
 		SHADER_PARAMETER(int32, BrickBatchOffset)
-		SHADER_PARAMETER(int32, VolumetricLightmapQualityMultiplier)
 		SHADER_PARAMETER_SRV(RaytracingAccelerationStructure, TLAS)
 		SHADER_PARAMETER_SRV(Buffer<uint4>, BrickRequests)
 		SHADER_PARAMETER_UAV(RWTexture3D<float3>, AmbientVector)
diff --git a/Engine/Plugins/Experimental/GPULightmass/Source/GPULightmass/Private/Scene/Scene.cpp b/Engine/Plugins/Experimental/GPULightmass/Source/GPULightmass/Private/Scene/Scene.cpp
--- a/Engine/Plugins/Experimental/GPULightmass/Source/GPULightmass/Private/Scene/Scene.cpp
+++ b/Engine/Plugins/Experimental/GPULightmass/Source/GPULightmass/Private/Scene/Scene.cpp
@@ -1627,8 +1627,10 @@ void FSceneRenderState::BackgroundTick()
 
 		{
 			int32 NumCellsPerBrick = 5 * 5 * 5;
-			SamplesTaken += VolumetricLightmapRenderer->SamplesTaken;
-			TotalSamples += (uint64)VolumetricLightmapRenderer->NumTotalBricks * NumCellsPerBrick * VolumetricLightmapRenderer->NumTotalPassesToRender;
+			// VLM traces 2 paths samples per pass
+			// Multiply by 2 to get a more linear progress
+			SamplesTaken += 2 * VolumetricLightmapRenderer->SamplesTaken;
+			TotalSamples += 2 * (uint64)VolumetricLightmapRenderer->NumTotalBricks * NumCellsPerBrick * VolumetricLightmapRenderer->NumTotalPassesToRender;
 		}
 
 		int32 IntPercentage = FMath::FloorToInt(SamplesTaken * 100.0 / TotalSamples);
diff --git a/Engine/Plugins/Experimental/GPULightmass/Source/GPULightmass/Private/VolumetricLightmap.cpp b/Engine/Plugins/Experimental/GPULightmass/Source/GPULightmass/Private/VolumetricLightmap.cpp
--- a/Engine/Plugins/Experimental/GPULightmass/Source/GPULightmass/Private/VolumetricLightmap.cpp
+++ b/Engine/Plugins/Experimental/GPULightmass/Source/GPULightmass/Private/VolumetricLightmap.cpp
@@ -84,14 +84,12 @@ FVolumetricLightmapRenderer::FVolumetricLightmapRenderer(FSceneRenderState* Scen
 {
 	VolumetricLightmap.Data = &VolumetricLightmapData;
 	
-	NumTotalPassesToRender = Scene->Settings->GISamples;
+	NumTotalPassesToRender = Scene->Settings->GISamples * Scene->Settings->VolumetricLightmapQualityMultiplier;
 	
 	if (Scene->Settings->bUseIrradianceCaching)
 	{
 		NumTotalPassesToRender += Scene->Settings->IrradianceCacheQuality;	
 	}
-
-	NumTotalPassesToRender *= Scene->Settings->VolumetricLightmapQualityMultiplier;
 }
 
 FPrecomputedVolumetricLightmap* FVolumetricLightmapRenderer::GetPrecomputedVolumetricLightmapForPreview()
@@ -586,7 +584,6 @@ void FVolumetricLightmapRenderer::BackgroundTick()
 					PassParameters->BrickRequests = BrickRequests.SRV;
 					PassParameters->NumTotalBricks = NumTotalBricks;
 					PassParameters->BrickBatchOffset = BrickBatchOffset;
-					PassParameters->VolumetricLightmapQualityMultiplier = Scene->Settings->VolumetricLightmapQualityMultiplier;
 					PassParameters->AmbientVector = AccumulationBrickData.AmbientVector.UAV;
 					PassParameters->SHCoefficients0R = AccumulationBrickData.SHCoefficients[0].UAV;
 					PassParameters->SHCoefficients1R = AccumulationBrickData.SHCoefficients[1].UAV;
@@ -740,9 +737,9 @@ void FVolumetricLightmapRenderer::BackgroundTick()
 
 					if (Scene->Settings->bUseIrradianceCaching)
 					{
-						if (EffectiveRenderPassIndex >= Scene->Settings->IrradianceCacheQuality * Scene->Settings->VolumetricLightmapQualityMultiplier)
+						if (EffectiveRenderPassIndex >= Scene->Settings->IrradianceCacheQuality)
 						{
-							EffectiveRenderPassIndex -= Scene->Settings->IrradianceCacheQuality * Scene->Settings->VolumetricLightmapQualityMultiplier;
+							EffectiveRenderPassIndex -= Scene->Settings->IrradianceCacheQuality;
 						}
 					}
 					
