GPULM: Disable lighting builds and show warning when Path Tracing is disabled in project settings

diff --git a/Engine/Plugins/Experimental/GPULightmass/Source/GPULightmassEditor/Private/GPULightmassEditorModule.cpp b/Engine/Plugins/Experimental/GPULightmass/Source/GPULightmassEditor/Private/GPULightmassEditorModule.cpp
--- a/Engine/Plugins/Experimental/GPULightmass/Source/GPULightmassEditor/Private/GPULightmassEditorModule.cpp
+++ b/Engine/Plugins/Experimental/GPULightmass/Source/GPULightmassEditor/Private/GPULightmassEditorModule.cpp
@@ -141,6 +141,12 @@ static bool IsStaticLightingAllowed()
 	return (AllowStaticLightingVar && AllowStaticLightingVar->GetValueOnAnyThread() > 0);
 }
 
+static bool IsPathTracingEnabled()
+{
+	static const auto CVarPathTracing = IConsoleManager::Get().FindTConsoleVariableDataInt(TEXT("r.PathTracing"));
+	return CVarPathTracing && CVarPathTracing->GetValueOnAnyThread() > 0;
+}
+
 void FGPULightmassEditorModule::StartupModule()
 {
 	FLevelEditorModule& LevelEditorModule = FModuleManager::LoadModuleChecked<FLevelEditorModule>(TEXT("LevelEditor"));
@@ -216,7 +222,7 @@ TSharedRef<SDockTab> FGPULightmassEditorModule::SpawnSettingsTab(const FSpawnTab
 					SNew(SButton)
 					.HAlign(HAlign_Center)
 					.ButtonStyle(FAppStyle::Get(), "FlatButton.Success")
-					.IsEnabled(IsRayTracingEnabled() && IsStaticLightingAllowed())
+					.IsEnabled(IsRayTracingEnabled() && IsStaticLightingAllowed() && IsPathTracingEnabled())
 					.Visibility_Lambda([](){ return IsRunning() ? EVisibility::Collapsed : EVisibility::Visible; })
 					.OnClicked_Raw(this, &FGPULightmassEditorModule::OnStartClicked)
 					[
@@ -341,6 +347,10 @@ TSharedRef<SDockTab> FGPULightmassEditorModule::SpawnSettingsTab(const FSpawnTab
 					{
 						return LOCTEXT("GPULightmassAllowStaticLightingDisabled", "GPU Lightmass requires Allow Static Lighting enabled in the project settings.");
 					}
+					else if (!IsPathTracingEnabled())
+					{
+						return LOCTEXT("GPULightmassPathTracingDisabled", "GPU Lightmass requires Path Tracing enabled in the project settings.");
+					}
 
 					// Ready
 					static FText ReadyMsg = FText(LOCTEXT("GPULightmassReady", "GPU Lightmass is ready."));
