GPULM: Complete rework of volumetric lightmap renderer

* Instead of allocating accumulation brick data for the entire scene and looping over the bricks across frames, now we progress through the bricks at 256 bricks/batch. This breaks interactive preview but drastically reduces GPU memory consumption (which means fewer GPU OOM crashes) and has much better irradiance cache locality
* The bricks are Morton curve sorted to further improve locality
* IC now guards against 0 samples entries. Modification to the cache outside IC building passes are prohibited to prevent meaningless cache thrashing when IC size is not enough to cover the entire scece
* VLM path tracing code is ported to RDG and manual transitions are removed

diff --git a/Engine/Plugins/Experimental/GPULightmass/Shaders/Private/BrickAllocationManagement.usf b/Engine/Plugins/Experimental/GPULightmass/Shaders/Private/BrickAllocationManagement.usf
--- a/Engine/Plugins/Experimental/GPULightmass/Shaders/Private/BrickAllocationManagement.usf
+++ b/Engine/Plugins/Experimental/GPULightmass/Shaders/Private/BrickAllocationManagement.usf
@@ -7,10 +7,12 @@
 #include "/Engine/Public/Platform.ush" 
 #include "BrickAllocationDefs.ush"
 #include "/Engine/Generated/UniformBuffers/VLMVoxelizationParams.ush"
+#include "/Engine/Private/MortonCode.ush"
  
 int3 VolumeSize;
 RWTexture3D<uint> VoxelizeVolume;
 int BrickSize;
+int MipLevel;
 int bIsHighestMip;
 
 [numthreads(4, 4, 4)]
@@ -71,6 +73,10 @@ void DilateVolumeCS(uint3 VoxelPos : SV_DispatchThreadID)
 RWTexture3D<uint4> IndirectionTexture;
 
 RWBuffer<uint4> BrickRequests;
+RWBuffer<uint4> BrickRequestsUnsorted;
+RWBuffer<uint> BrickRequestKeys;
+RWBuffer<uint> BrickRequestSortedIndices;
+RWBuffer<uint> BrickRequestSortedIndicesInverse;
 
 [numthreads(4, 4, 4)]
 void CountNumBricksCS(
@@ -100,7 +106,10 @@ void GatherBrickRequestsCS(
 	if (any(VoxelPos >= VolumeSize)) return; 
 	if (VoxelizeVolume[VoxelPos])
 	{
-		BrickRequests[VoxelizeVolume[VoxelPos] - 1] = uint4(VoxelPos, BrickSize);
+		int LinearIndex = VoxelizeVolume[VoxelPos] - 1;
+		BrickRequestKeys[LinearIndex] = (MipLevel << 30) | MortonCode3(VoxelPos.x) | (MortonCode3(VoxelPos.y) << 1) | (MortonCode3(VoxelPos.z) << 2);
+		BrickRequestSortedIndices[LinearIndex] = LinearIndex;
+		BrickRequestsUnsorted[LinearIndex] = uint4(VoxelPos, BrickSize);
 	}
 }
 
@@ -123,6 +132,26 @@ void SplatVolumeCS(
 	}
 }
 
+[numthreads(1, 1, 1)]
+void WriteSortedBrickRequestsCS(uint3 RequestIndex : SV_DispatchThreadID)
+{
+	BrickRequests[RequestIndex.x] = BrickRequestsUnsorted[BrickRequestSortedIndices[RequestIndex.x]];
+	BrickRequestSortedIndicesInverse[BrickRequestSortedIndices[RequestIndex.x]] = RequestIndex.x;
+}
+
+[numthreads(4, 4, 4)]
+void PermuteVoxelizeVolumeCS(
+	uint3 VoxelPos : SV_DispatchThreadID
+	) 
+{
+	if (any(VoxelPos >= VolumeSize)) return; 
+	if (VoxelizeVolume[VoxelPos])
+	{
+		int LinearIndex = VoxelizeVolume[VoxelPos] - 1;
+		VoxelizeVolume[VoxelPos] = BrickRequestSortedIndicesInverse[LinearIndex] + 1;
+	}
+}
+
 int3 IndirectionTextureDim;
 int3 BrickDataDimensions;
 
@@ -148,7 +177,7 @@ RWTexture3D<UNORM float> OutDirectionalLightShadowing;
 
 uint FrameNumber;
 
-float3 DilateWithValidityMask(RWTexture3D<float3> Texture, Texture3D<float4> CoverageMask, float3 UVW)
+float3 DilateWithValidityMask(RWTexture3D<float3> Texture, RWTexture3D<uint> CoverageMask, float3 UVW)
 {
 	float AccumulatedWeight = 0.0f;
 	float3 Value = float3(0, 0, 0);
@@ -177,7 +206,7 @@ float3 DilateWithValidityMask(RWTexture3D<float3> Texture, Texture3D<float4> Cov
 				float3 BrickUVW = IndirectionTextureValue.xyz * PaddedBrickSize + frac(IndirectionTextureUVW / IndirectionTextureValue.w) * 4;
 			
 				int3 Coord = BrickUVW;
-				if (asuint(CoverageMask[Coord].w) >= FrameNumber / 3)
+				if (CoverageMask[Coord] == 1)
 				{
 					Value += Texture[Coord];
 					AccumulatedWeight += 1;
@@ -189,7 +218,7 @@ float3 DilateWithValidityMask(RWTexture3D<float3> Texture, Texture3D<float4> Cov
 	return Value / max(AccumulatedWeight, 0.00001f);
 }
 
-UNORM float4 DilateWithValidityMask(RWTexture3D<UNORM float4> Texture, Texture3D<float4> CoverageMask, float3 UVW)
+UNORM float4 DilateWithValidityMask(RWTexture3D<UNORM float4> Texture, RWTexture3D<uint> CoverageMask, float3 UVW)
 {
 	int3 Coord = UVW;
 	float3 CoordinateFraction = frac(UVW);
@@ -220,7 +249,7 @@ UNORM float4 DilateWithValidityMask(RWTexture3D<UNORM float4> Texture, Texture3D
 				float3 BrickUVW = IndirectionTextureValue.xyz * PaddedBrickSize + frac(IndirectionTextureUVW / IndirectionTextureValue.w) * 4;
 			
 				int3 Coord = BrickUVW;
-				if (asuint(CoverageMask[Coord].w) >= FrameNumber / 3)
+				if (CoverageMask[Coord] == 1)
 				{
 					Value += Texture[Coord];
 					AccumulatedWeight += 1;
@@ -232,7 +261,7 @@ UNORM float4 DilateWithValidityMask(RWTexture3D<UNORM float4> Texture, Texture3D
 	return Value / max(AccumulatedWeight, 0.00001f);
 }
 
-UNORM float DilateWithValidityMask(RWTexture3D<UNORM float> Texture, Texture3D<float4> CoverageMask, float3 UVW)
+UNORM float DilateWithValidityMask(RWTexture3D<UNORM float> Texture, RWTexture3D<uint> CoverageMask, float3 UVW)
 {
 	int3 Coord = UVW;
 	float3 CoordinateFraction = frac(UVW);
@@ -263,7 +292,7 @@ UNORM float DilateWithValidityMask(RWTexture3D<UNORM float> Texture, Texture3D<f
 				float3 BrickUVW = IndirectionTextureValue.xyz * PaddedBrickSize + frac(IndirectionTextureUVW / IndirectionTextureValue.w) * 4;
 			
 				int3 Coord = BrickUVW;
-				if (asuint(CoverageMask[Coord].w) >= FrameNumber / 3)
+				if (CoverageMask[Coord] == 1)
 				{
 					Value += Texture[Coord];
 					AccumulatedWeight += 1;
@@ -349,6 +378,7 @@ float3 ManualTrilinearFilter(RWTexture3D<UNORM float> Texture, float3 UVW)
 
 int NumTotalBricks;
 int BrickBatchOffset;
+RWTexture3D<uint> ValidityMask;
 
 [numthreads(5, 5, 5)]
 void StitchBorderCS(uint3 BrickID : SV_GroupID, uint3 CellPosInBrick : SV_GroupThreadID)
@@ -360,19 +390,18 @@ void StitchBorderCS(uint3 BrickID : SV_GroupID, uint3 CellPosInBrick : SV_GroupT
 	float PaddedBrickSize = 4 + 1;
 	uint3 VoxelPos = ComputeBrickLayoutPosition(BrickIndex, uint3(256, 256, 256)) * PaddedBrickSize + CellPosInBrick;
 
-	uint SampleCount = asuint(AmbientVector[VoxelPos].w);
-	if (FrameNumber >= 32 && SampleCount < FrameNumber / 3)
+	if (ValidityMask[VoxelPos] == 0)
 	{
 		float3 IndirectionTextureUVW = CellPosInVLM * InvCellSizeInVLM;
-		OutAmbientVector[VoxelPos] = DilateWithValidityMask(OutAmbientVector, AmbientVector, IndirectionTextureUVW);
-		OutSHCoefficients0R[VoxelPos] = DilateWithValidityMask(OutSHCoefficients0R, AmbientVector, IndirectionTextureUVW);
-		OutSHCoefficients1R[VoxelPos] = DilateWithValidityMask(OutSHCoefficients1R, AmbientVector, IndirectionTextureUVW);
-		OutSHCoefficients0G[VoxelPos] = DilateWithValidityMask(OutSHCoefficients0G, AmbientVector, IndirectionTextureUVW);
-		OutSHCoefficients1G[VoxelPos] = DilateWithValidityMask(OutSHCoefficients1G, AmbientVector, IndirectionTextureUVW);
-		OutSHCoefficients0B[VoxelPos] = DilateWithValidityMask(OutSHCoefficients0B, AmbientVector, IndirectionTextureUVW);
-		OutSHCoefficients1B[VoxelPos] = DilateWithValidityMask(OutSHCoefficients1B, AmbientVector, IndirectionTextureUVW);
-		OutSkyBentNormal[VoxelPos] = DilateWithValidityMask(OutSkyBentNormal, AmbientVector, IndirectionTextureUVW);
-		OutDirectionalLightShadowing[VoxelPos] = DilateWithValidityMask(OutDirectionalLightShadowing, AmbientVector, IndirectionTextureUVW);
+		OutAmbientVector[VoxelPos] = DilateWithValidityMask(OutAmbientVector, ValidityMask, IndirectionTextureUVW);
+		OutSHCoefficients0R[VoxelPos] = DilateWithValidityMask(OutSHCoefficients0R, ValidityMask, IndirectionTextureUVW);
+		OutSHCoefficients1R[VoxelPos] = DilateWithValidityMask(OutSHCoefficients1R, ValidityMask, IndirectionTextureUVW);
+		OutSHCoefficients0G[VoxelPos] = DilateWithValidityMask(OutSHCoefficients0G, ValidityMask, IndirectionTextureUVW);
+		OutSHCoefficients1G[VoxelPos] = DilateWithValidityMask(OutSHCoefficients1G, ValidityMask, IndirectionTextureUVW);
+		OutSHCoefficients0B[VoxelPos] = DilateWithValidityMask(OutSHCoefficients0B, ValidityMask, IndirectionTextureUVW);
+		OutSHCoefficients1B[VoxelPos] = DilateWithValidityMask(OutSHCoefficients1B, ValidityMask, IndirectionTextureUVW);
+		OutSkyBentNormal[VoxelPos] = DilateWithValidityMask(OutSkyBentNormal, ValidityMask, IndirectionTextureUVW);
+		OutDirectionalLightShadowing[VoxelPos] = DilateWithValidityMask(OutDirectionalLightShadowing, ValidityMask, IndirectionTextureUVW);
 	}
 
 	for (int dx = -1; dx <= 1; dx++)
@@ -403,6 +432,8 @@ void StitchBorderCS(uint3 BrickID : SV_GroupID, uint3 CellPosInBrick : SV_GroupT
 	}
 }
 
+int NumTotalPassesToRender;
+
 [numthreads(5, 5, 5)]
 void FinalizeBrickResultsCS(uint3 BrickID : SV_GroupID, uint3 CellPosInBrick : SV_GroupThreadID)
 {
@@ -411,12 +442,13 @@ void FinalizeBrickResultsCS(uint3 BrickID : SV_GroupID, uint3 CellPosInBrick : S
 	int3 CellPosInVLM = (BrickRequests[BrickIndex].xyz * 4 + CellPosInBrick) * BrickRequests[BrickIndex].w;	
 	float3 InvCellSizeInVLM = float3(1, 1, 1) / 4.0f;
 	float PaddedBrickSize = 4 + 1;
-	uint3 VoxelPos = ComputeBrickLayoutPosition(BrickIndex, uint3(256, 256, 256)) * PaddedBrickSize + CellPosInBrick;
+	uint3 VoxelPos = ComputeBrickLayoutPosition(BrickID.x, uint3(256, 256, 256)) * PaddedBrickSize + CellPosInBrick;
+	uint3 OutVoxelPos = ComputeBrickLayoutPosition(BrickIndex, uint3(256, 256, 256)) * PaddedBrickSize + CellPosInBrick;
 
 	uint SampleCount = asuint(AmbientVector[VoxelPos].w);
 	if (SampleCount > 0)
 	{
-		OutAmbientVector[VoxelPos] = AmbientVector[VoxelPos].xyz / SampleCount;
+		OutAmbientVector[OutVoxelPos] = AmbientVector[VoxelPos].xyz / SampleCount;
 		
 		float3 InvAmbient = float3(1.0f, 1.0f, 1.0f) / max(AmbientVector[VoxelPos].xyz / SampleCount, float3(0.0001f, 0.0001f, 0.0001f));
 		
@@ -432,13 +464,15 @@ void FinalizeBrickResultsCS(uint3 BrickID : SV_GroupID, uint3 CellPosInBrick : S
 			0.282095f / 1.092548f,
 			0.282095f / (2.0f * 0.546274f));
 	
-		OutSHCoefficients0R[VoxelPos] = (SHCoefficients0R[VoxelPos] / SampleCount) * InvAmbient.r * CoefficientNormalizationScale0 * 0.5f + 0.5f;
-		OutSHCoefficients1R[VoxelPos] = (SHCoefficients1R[VoxelPos] / SampleCount) * InvAmbient.r * CoefficientNormalizationScale1 * 0.5f + 0.5f;
-		OutSHCoefficients0G[VoxelPos] = (SHCoefficients0G[VoxelPos] / SampleCount) * InvAmbient.g * CoefficientNormalizationScale0 * 0.5f + 0.5f;
-		OutSHCoefficients1G[VoxelPos] = (SHCoefficients1G[VoxelPos] / SampleCount) * InvAmbient.g * CoefficientNormalizationScale1 * 0.5f + 0.5f;
-		OutSHCoefficients0B[VoxelPos] = (SHCoefficients0B[VoxelPos] / SampleCount) * InvAmbient.b * CoefficientNormalizationScale0 * 0.5f + 0.5f;
-		OutSHCoefficients1B[VoxelPos] = (SHCoefficients1B[VoxelPos] / SampleCount) * InvAmbient.b * CoefficientNormalizationScale1 * 0.5f + 0.5f;
-		OutSkyBentNormal[VoxelPos] = SkyBentNormal[VoxelPos] / SampleCount * 0.5f + 0.5f;
-		OutDirectionalLightShadowing[VoxelPos] = DirectionalLightShadowing[VoxelPos];
+		OutSHCoefficients0R[OutVoxelPos] = (SHCoefficients0R[VoxelPos] / SampleCount) * InvAmbient.r * CoefficientNormalizationScale0 * 0.5f + 0.5f;
+		OutSHCoefficients1R[OutVoxelPos] = (SHCoefficients1R[VoxelPos] / SampleCount) * InvAmbient.r * CoefficientNormalizationScale1 * 0.5f + 0.5f;
+		OutSHCoefficients0G[OutVoxelPos] = (SHCoefficients0G[VoxelPos] / SampleCount) * InvAmbient.g * CoefficientNormalizationScale0 * 0.5f + 0.5f;
+		OutSHCoefficients1G[OutVoxelPos] = (SHCoefficients1G[VoxelPos] / SampleCount) * InvAmbient.g * CoefficientNormalizationScale1 * 0.5f + 0.5f;
+		OutSHCoefficients0B[OutVoxelPos] = (SHCoefficients0B[VoxelPos] / SampleCount) * InvAmbient.b * CoefficientNormalizationScale0 * 0.5f + 0.5f;
+		OutSHCoefficients1B[OutVoxelPos] = (SHCoefficients1B[VoxelPos] / SampleCount) * InvAmbient.b * CoefficientNormalizationScale1 * 0.5f + 0.5f;
+		OutSkyBentNormal[OutVoxelPos] = SkyBentNormal[VoxelPos] / SampleCount * 0.5f + 0.5f;
+		OutDirectionalLightShadowing[OutVoxelPos] = DirectionalLightShadowing[VoxelPos];
 	}
+
+	ValidityMask[OutVoxelPos] = SampleCount > NumTotalPassesToRender / 3;
 }
diff --git a/Engine/Plugins/Experimental/GPULightmass/Shaders/Private/LightmapPathTracing.usf b/Engine/Plugins/Experimental/GPULightmass/Shaders/Private/LightmapPathTracing.usf
--- a/Engine/Plugins/Experimental/GPULightmass/Shaders/Private/LightmapPathTracing.usf
+++ b/Engine/Plugins/Experimental/GPULightmass/Shaders/Private/LightmapPathTracing.usf
@@ -450,18 +450,23 @@ void PathTracingKernel(
 				float DistanceFraction = CellArea / (2 * PI * DistanceFromTexelToHit * DistanceFromTexelToHit);
 			
 				// We accept a cache entry, if:
-				// 1) We're outside of IC building passes. An entry with low sample count means we just need that much - still, we're going to require a minimum of 8 samples to remove outliers
-				// 2) The entry has been built by other lightmap tiles and contains more samples than we will ever need.
+				// 1) The entry has been built by other lightmap tiles and contains more samples than we will ever need.
+				// 2) We're outside of IC building passes. An entry with low sample count means we just need that much - still, we're going to require a minimum of 4 samples to remove outliers from potential cache thrashing
 
 				bool bShouldAcceptCacheEntry = false;
 
-			#ifdef LIGHTMAP_PATH_TRACING_MAIN_RG
-				// Surface lightmaps
-				bShouldAcceptCacheEntry = RenderPassIndex >= GetIrradianceCachingQuality() || RecordIrradianceAndSampleCount.w >= TotalSamplesOverAllICPasses * DistanceFraction;
-			#else
-				// Volumetric lightmaps
-				bShouldAcceptCacheEntry = RenderPassIndex >= GetIrradianceCachingQuality();
-			#endif
+				uint SampleCountThreshold;
+
+				if (RenderPassIndex < GetIrradianceCachingQuality())
+				{
+					SampleCountThreshold = TotalSamplesOverAllICPasses * DistanceFraction;
+				}
+				else
+				{
+					SampleCountThreshold = 4;
+				}
+
+				bShouldAcceptCacheEntry = RecordIrradianceAndSampleCount.w >= SampleCountThreshold;
 				
 			#if IC_BACKFACE_DETECTION
 				// Reject the entry if it can see back faces
@@ -501,30 +506,27 @@ void PathTracingKernel(
 				// Failed query.
 				if (Bounce == 1)
 				{
-					if (!bGeometryQuerySuccessful)
+					// Prevent modification to the cache outside IC building passes
+					if (RenderPassIndex < GetIrradianceCachingQuality())
 					{
-						bool bShouldPlaceNewEntry = false;
-					
-					#ifdef LIGHTMAP_PATH_TRACING_MAIN_RG
-						// Surface lightmaps are relatively high density. Only use the very first few passes to place new entries to reduce cache thrashing
-						bShouldPlaceNewEntry = RenderPassIndex < 16;
-					#else
-						// Always place new entry for volumetric lightmaps
-						bShouldPlaceNewEntry = true;
-					#endif
-						
-						if (bShouldPlaceNewEntry)
+						if (!bGeometryQuerySuccessful)
 						{
-							bShouldEmitGeometryHitPoint = true;
-							IrradianceCacheFinalGatherHitPoint.WorldPosition = HitPosition;
-							IrradianceCacheFinalGatherHitPoint.WorldNormal = Payload.WorldNormal;
-							NewEntrySize = IrradianceCachingParameters.Spacing;
-							break;
+							// Only use the very first few passes to place new entries to reduce cache thrashing
+							bool bShouldPlaceNewEntry = RenderPassIndex < 16;
+
+							if (bShouldPlaceNewEntry)
+							{
+								bShouldEmitGeometryHitPoint = true;
+								IrradianceCacheFinalGatherHitPoint.WorldPosition = HitPosition;
+								IrradianceCacheFinalGatherHitPoint.WorldNormal = Payload.WorldNormal;
+								NewEntrySize = IrradianceCachingParameters.Spacing;
+								break;
+							}
+						}
+						else
+						{
+							NearestCacheEntryID = NearestRecordIndex;
 						}
-					}
-					else
-					{
-						NearestCacheEntryID = NearestRecordIndex;
 					}
 				}
 			}
@@ -1219,7 +1221,7 @@ void VolumetricLightmapPathTracingMainRG()
 	uint CellIndex = DispatchRaysIndex().x % BrickVolumeSize;
 	uint3 CellPosInBrick = ComputeBrickLayoutPosition(CellIndex, uint3(5, 5, 5));
 	uint PaddedBrickSize = 4 + 1;
-	uint3 VoxelPos = ComputeBrickLayoutPosition(BrickIndex, uint3(256, 256, 256)) * PaddedBrickSize + CellPosInBrick;
+	uint3 VoxelPos = ComputeBrickLayoutPosition(DispatchRaysIndex().x / BrickVolumeSize, uint3(256, 256, 256)) * PaddedBrickSize + CellPosInBrick;
 	
 	bool bIsValidSample = true;	
 	
diff --git a/Engine/Plugins/Experimental/GPULightmass/Source/GPULightmass/Private/LightmapRayTracing.h b/Engine/Plugins/Experimental/GPULightmass/Source/GPULightmass/Private/LightmapRayTracing.h
--- a/Engine/Plugins/Experimental/GPULightmass/Source/GPULightmass/Private/LightmapRayTracing.h
+++ b/Engine/Plugins/Experimental/GPULightmass/Source/GPULightmass/Private/LightmapRayTracing.h
@@ -134,15 +134,15 @@ class FVolumetricLightmapPathTracingRGS : public FGlobalShader
 		SHADER_PARAMETER(int32, BrickBatchOffset)
 		SHADER_PARAMETER_SRV(RaytracingAccelerationStructure, TLAS)
 		SHADER_PARAMETER_SRV(Buffer<uint4>, BrickRequests)
-		SHADER_PARAMETER_UAV(RWTexture3D<float3>, AmbientVector)
-		SHADER_PARAMETER_UAV(RWTexture3D<float4>, SHCoefficients0R)
-		SHADER_PARAMETER_UAV(RWTexture3D<float4>, SHCoefficients1R)
-		SHADER_PARAMETER_UAV(RWTexture3D<float4>, SHCoefficients0G)
-		SHADER_PARAMETER_UAV(RWTexture3D<float4>, SHCoefficients1G)
-		SHADER_PARAMETER_UAV(RWTexture3D<float4>, SHCoefficients0B)
-		SHADER_PARAMETER_UAV(RWTexture3D<float4>, SHCoefficients1B)
-		SHADER_PARAMETER_UAV(RWTexture3D<float4>, SkyBentNormal)
-		SHADER_PARAMETER_UAV(RWTexture3D<float>, DirectionalLightShadowing)
+		SHADER_PARAMETER_RDG_TEXTURE_UAV(RWTexture3D<float3>, AmbientVector)
+		SHADER_PARAMETER_RDG_TEXTURE_UAV(RWTexture3D<float4>, SHCoefficients0R)
+		SHADER_PARAMETER_RDG_TEXTURE_UAV(RWTexture3D<float4>, SHCoefficients1R)
+		SHADER_PARAMETER_RDG_TEXTURE_UAV(RWTexture3D<float4>, SHCoefficients0G)
+		SHADER_PARAMETER_RDG_TEXTURE_UAV(RWTexture3D<float4>, SHCoefficients1G)
+		SHADER_PARAMETER_RDG_TEXTURE_UAV(RWTexture3D<float4>, SHCoefficients0B)
+		SHADER_PARAMETER_RDG_TEXTURE_UAV(RWTexture3D<float4>, SHCoefficients1B)
+		SHADER_PARAMETER_RDG_TEXTURE_UAV(RWTexture3D<float4>, SkyBentNormal)
+		SHADER_PARAMETER_RDG_TEXTURE_UAV(RWTexture3D<float>, DirectionalLightShadowing)
 
 		SHADER_PARAMETER_RDG_BUFFER_SRV(StructuredBuffer<FPathTracingLight>, SceneLights)
 		SHADER_PARAMETER(uint32, SceneLightCount)
diff --git a/Engine/Plugins/Experimental/GPULightmass/Source/GPULightmass/Private/LightmapRenderer.cpp b/Engine/Plugins/Experimental/GPULightmass/Source/GPULightmass/Private/LightmapRenderer.cpp
--- a/Engine/Plugins/Experimental/GPULightmass/Source/GPULightmass/Private/LightmapRenderer.cpp
+++ b/Engine/Plugins/Experimental/GPULightmass/Source/GPULightmass/Private/LightmapRenderer.cpp
@@ -258,10 +258,7 @@ FLightmapRenderer::FLightmapRenderer(FSceneRenderState* InScene)
 		TilesVisibleLastFewFrames.AddDefaulted(60);
 	}
 
-	if (Scene->Settings->bVisualizeIrradianceCache)
-	{
-		IrradianceCacheVisualizationDelegateHandle = GetRendererModule().RegisterPostOpaqueRenderDelegate(FPostOpaqueRenderDelegate::CreateRaw(this, &FLightmapRenderer::RenderIrradianceCacheVisualization));
-	}
+	IrradianceCacheVisualizationDelegateHandle = GetRendererModule().RegisterPostOpaqueRenderDelegate(FPostOpaqueRenderDelegate::CreateRaw(this, &FLightmapRenderer::RenderIrradianceCacheVisualization));
 }
 
 FLightmapRenderer::~FLightmapRenderer()
@@ -1511,6 +1508,7 @@ void FLightmapRenderer::Finalize(FRDGBuilder& GraphBuilder)
 
 	if (bIsCompilingShaders)
 	{
+		PendingTileRequests.Empty();
 		return;
 	}
 	
@@ -3703,16 +3701,6 @@ void FLightmapRenderer::BackgroundTick()
 		// Purge resources when 'realtime' is not checked on editor viewport to avoid leak & slowing down
 		RHICmdList.ImmediateFlush(EImmediateFlushType::FlushRHIThreadFlushResources);
 	}
-
-	if (Scene->Settings->bVisualizeIrradianceCache && !IrradianceCacheVisualizationDelegateHandle.IsValid())
-	{
-		IrradianceCacheVisualizationDelegateHandle = GetRendererModule().RegisterPostOpaqueRenderDelegate(FPostOpaqueRenderDelegate::CreateRaw(this, &FLightmapRenderer::RenderIrradianceCacheVisualization));
-	}
-	else if (!Scene->Settings->bVisualizeIrradianceCache && IrradianceCacheVisualizationDelegateHandle.IsValid())
-	{
-		GetRendererModule().RemovePostOpaqueRenderDelegate(IrradianceCacheVisualizationDelegateHandle);
-		IrradianceCacheVisualizationDelegateHandle.Reset();
-	}
 }
 
 void FLightmapRenderer::BumpRevision()
@@ -3782,6 +3770,11 @@ void FLightmapRenderer::DeduplicateRecordedTileRequests()
 
 void FLightmapRenderer::RenderIrradianceCacheVisualization(FPostOpaqueRenderParameters& Parameters)
 {
+	if (!Scene->Settings->bVisualizeIrradianceCache)
+	{
+		return;
+	}
+	
 	FRDGBuilder& GraphBuilder = *Parameters.GraphBuilder;
 	const ERHIFeatureLevel::Type FeatureLevel = Scene->FeatureLevel;
 
diff --git a/Engine/Plugins/Experimental/GPULightmass/Source/GPULightmass/Private/Scene/Scene.cpp b/Engine/Plugins/Experimental/GPULightmass/Source/GPULightmass/Private/Scene/Scene.cpp
--- a/Engine/Plugins/Experimental/GPULightmass/Source/GPULightmass/Private/Scene/Scene.cpp
+++ b/Engine/Plugins/Experimental/GPULightmass/Source/GPULightmass/Private/Scene/Scene.cpp
@@ -1629,7 +1629,7 @@ void FSceneRenderState::BackgroundTick()
 			int32 NumCellsPerBrick = 5 * 5 * 5;
 			// VLM traces 2 paths samples per pass
 			// Multiply by 2 to get a more linear progress
-			SamplesTaken += 2 * VolumetricLightmapRenderer->SamplesTaken;
+			SamplesTaken += 2 * VolumetricLightmapRenderer->SamplesTaken * NumCellsPerBrick;
 			TotalSamples += 2 * (uint64)VolumetricLightmapRenderer->NumTotalBricks * NumCellsPerBrick * VolumetricLightmapRenderer->NumTotalPassesToRender;
 		}
 
diff --git a/Engine/Plugins/Experimental/GPULightmass/Source/GPULightmass/Private/VolumetricLightmap.cpp b/Engine/Plugins/Experimental/GPULightmass/Source/GPULightmass/Private/VolumetricLightmap.cpp
--- a/Engine/Plugins/Experimental/GPULightmass/Source/GPULightmass/Private/VolumetricLightmap.cpp
+++ b/Engine/Plugins/Experimental/GPULightmass/Source/GPULightmass/Private/VolumetricLightmap.cpp
@@ -11,6 +11,7 @@
 #include "LevelEditorViewport.h"
 #include "Editor.h"
 #include "RectLightTextureManager.h"
+#include "GPUSort.h"
 
 IMPLEMENT_MATERIAL_SHADER_TYPE(, FVLMVoxelizationVS, TEXT("/Plugin/GPULightmass/Private/VolumetricLightmapVoxelization.usf"), TEXT("VLMVoxelizationVS"), SF_Vertex);
 IMPLEMENT_MATERIAL_SHADER_TYPE(, FVLMVoxelizationGS, TEXT("/Plugin/GPULightmass/Private/VolumetricLightmapVoxelization.usf"), TEXT("VLMVoxelizationGS"), SF_Geometry);
@@ -24,6 +25,8 @@ IMPLEMENT_GLOBAL_SHADER(FGatherBrickRequestsCS, "/Plugin/GPULightmass/Private/Br
 IMPLEMENT_GLOBAL_SHADER(FSplatVolumeCS, "/Plugin/GPULightmass/Private/BrickAllocationManagement.usf", "SplatVolumeCS", SF_Compute);
 IMPLEMENT_GLOBAL_SHADER(FStitchBorderCS, "/Plugin/GPULightmass/Private/BrickAllocationManagement.usf", "StitchBorderCS", SF_Compute);
 IMPLEMENT_GLOBAL_SHADER(FFinalizeBrickResultsCS, "/Plugin/GPULightmass/Private/BrickAllocationManagement.usf", "FinalizeBrickResultsCS", SF_Compute);
+IMPLEMENT_GLOBAL_SHADER(FWriteSortedBrickRequestsCS, "/Plugin/GPULightmass/Private/BrickAllocationManagement.usf", "WriteSortedBrickRequestsCS", SF_Compute);
+IMPLEMENT_GLOBAL_SHADER(FPermuteVoxelizeVolumeCS, "/Plugin/GPULightmass/Private/BrickAllocationManagement.usf", "PermuteVoxelizeVolumeCS", SF_Compute);
 
 IMPLEMENT_STATIC_UNIFORM_BUFFER_STRUCT(FVLMVoxelizationParams, "VLMVoxelizationParams", SceneTextures);
 
@@ -52,6 +55,15 @@ void InitializeBrickData(FIntVector BrickDataDimensions, FVolumetricLightmapBric
 
 	BrickData.DirectionalLightShadowing.CreateTargetTexture(BrickDataDimensions);
 	BrickData.DirectionalLightShadowing.CreateUAV();
+
+	size_t TotalSize = 0;
+	TotalSize += GPixelFormats[BrickData.AmbientVector.Format].BlockBytes;
+	TotalSize += UE_ARRAY_COUNT(BrickData.SHCoefficients) * GPixelFormats[BrickData.SHCoefficients[0].Format].BlockBytes;
+	TotalSize += GPixelFormats[BrickData.SkyBentNormal.Format].BlockBytes;
+	TotalSize += GPixelFormats[BrickData.DirectionalLightShadowing.Format].BlockBytes;
+	TotalSize *= BrickDataDimensions.X * BrickDataDimensions.Y * BrickDataDimensions.Z;
+	
+	UE_LOG(LogGPULightmass, Log, TEXT("Allocated %.2fMB for volumetric lightmap %s brick data"), TotalSize / 1024.0f / 1024.0f, bForAccumulation ? TEXT("accumulation") : TEXT("display"));
 }
 
 void ReleaseBrickData(FVolumetricLightmapBrickData& BrickData)
@@ -104,6 +116,17 @@ BEGIN_SHADER_PARAMETER_STRUCT(FVoxelizeMeshPassParameters, )
 	RENDER_TARGET_BINDING_SLOTS()
 END_SHADER_PARAMETER_STRUCT()
 
+BEGIN_SHADER_PARAMETER_STRUCT(FSortBrickRequestsPassParameters, )
+	SHADER_PARAMETER_RDG_BUFFER_SRV(Buffer<uint>, BrickRequestKeysSRV)
+	SHADER_PARAMETER_RDG_BUFFER_SRV(Buffer<uint>, BrickRequestKeysSRV2)
+	SHADER_PARAMETER_RDG_BUFFER_UAV(RWBuffer<uint>, BrickRequestKeysUAV)
+	SHADER_PARAMETER_RDG_BUFFER_UAV(RWBuffer<uint>, BrickRequestKeysUAV2)
+	SHADER_PARAMETER_RDG_BUFFER_SRV(Buffer<uint>, BrickRequestSortedIndicesSRV)
+	SHADER_PARAMETER_RDG_BUFFER_SRV(Buffer<uint>, BrickRequestSortedIndicesSRV2)
+	SHADER_PARAMETER_RDG_BUFFER_UAV(RWBuffer<uint>, BrickRequestSortedIndicesUAV)
+	SHADER_PARAMETER_RDG_BUFFER_UAV(RWBuffer<uint>, BrickRequestSortedIndicesUAV2)
+END_SHADER_PARAMETER_STRUCT()
+	
 void FVolumetricLightmapRenderer::VoxelizeScene()
 {
 	FGlobalShaderMap* GlobalShaderMap = GetGlobalShaderMap(Scene->FeatureLevel);
@@ -114,6 +137,7 @@ void FVolumetricLightmapRenderer::VoxelizeScene()
 	}
 
 	IndirectionTexture.SafeRelease();
+	ValidityBrickData.SafeRelease();
 
 	ReleaseBrickData(VolumetricLightmapData.BrickData);
 	ReleaseBrickData(AccumulationBrickData);
@@ -135,15 +159,6 @@ void FVolumetricLightmapRenderer::VoxelizeScene()
 	VolumeSize = FVector(TopLevelGridSize) * DetailCellsPerTopLevelBrick * TargetDetailCellSize;
 	FBox FinalVolume(VolumeMin, VolumeMin + VolumeSize);
 
-	UE_LOG(LogGPULightmass, Log, TEXT("Volumetric lightmap voxelization bounds set to (%.2f, %.2f, %.2f) - (%.2f, %.2f, %.2f)"),
-		FinalVolume.Min.X,
-		FinalVolume.Min.Y,
-		FinalVolume.Min.Z,
-		FinalVolume.Max.X,
-		FinalVolume.Max.Y,
-		FinalVolume.Max.Z
-	);
-
 	const int32 IndirectionCellsPerTopLevelCell = DetailCellsPerTopLevelBrick / BrickSize;
 
 	IndirectionTextureDimensions = TopLevelGridSize * IndirectionCellsPerTopLevelCell;
@@ -168,6 +183,9 @@ void FVolumetricLightmapRenderer::VoxelizeScene()
 			FRDGTexture* Texture = GraphBuilder.CreateTexture(Desc, TEXT("GPULightmassVLMIndirectionTexture"));
 			IndirectionTexture = GraphBuilder.ConvertToExternalTexture(Texture);
 			IndirectionTextureUAV = GraphBuilder.CreateUAV(Texture);
+
+			UE_LOG(LogGPULightmass, Log, TEXT("Allocated %.2fMB for volumetric lightmap indirection texture"),
+				IndirectionTextureDimensions.X * IndirectionTextureDimensions.Y * IndirectionTextureDimensions.Z * GPixelFormats[Desc.Format].BlockBytes / 1024.0f / 1024.0f);
 		}
 
 		TArray<FRDGTextureUAV*, FRDGArrayAllocator> VoxelizationVolumeMipUAVs;
@@ -438,7 +456,7 @@ void FVolumetricLightmapRenderer::VoxelizeScene()
 	}
 
 	InitializeBrickData(BrickLayoutDimensions * 5, VolumetricLightmapData.BrickData, false);
-	InitializeBrickData(BrickLayoutDimensions * 5, AccumulationBrickData, true);
+	InitializeBrickData(FIntVector(256, 1, 1) * 5, AccumulationBrickData, true);
 	BrickRequests.Initialize(TEXT("BrickRequests"), 16, NumTotalBricks, PF_R32G32B32A32_UINT, BUF_UnorderedAccess);
 
 	VolumetricLightmapData.BrickDataDimensions = BrickLayoutDimensions * 5;
@@ -446,7 +464,24 @@ void FVolumetricLightmapRenderer::VoxelizeScene()
 	{
 		FRDGBuilder GraphBuilder(RHICmdList);
 
+		FRDGTextureRef ValidityMask = GraphBuilder.CreateTexture(FRDGTextureDesc::Create3D(VolumetricLightmapData.BrickDataDimensions,
+															   PF_R8_UINT,
+															   FClearValueBinding::Black,
+															   ETextureCreateFlags::ShaderResource | ETextureCreateFlags::RenderTargetable | ETextureCreateFlags::UAV), TEXT("ValidityBrickData"));
+		ValidityBrickData = GraphBuilder.ConvertToExternalTexture(ValidityMask);
+
+		AddClearUAVPass(GraphBuilder, GraphBuilder.CreateUAV(ValidityMask), 0u);
+
+		UE_LOG(LogGPULightmass, Log, TEXT("Allocated %.2fMB for volumetric lightmap validity mask"),
+			VolumetricLightmapData.BrickDataDimensions.X * VolumetricLightmapData.BrickDataDimensions.Y * VolumetricLightmapData.BrickDataDimensions.Z * GPixelFormats[PF_R8_UINT].BlockBytes / 1024.0f / 1024.0f);
+		
 		FRDGTextureUAV* IndirectionTextureUAV = GraphBuilder.CreateUAV(GraphBuilder.RegisterExternalTexture(IndirectionTexture));
+		FRDGBufferRef BrickRequestsUnsorted = GraphBuilder.CreateBuffer(FRDGBufferDesc::CreateBufferDesc(sizeof(FIntVector4), NumTotalBricks), TEXT("BrickRequestsUnsorted"));
+		FRDGBufferRef BrickRequestKeys = GraphBuilder.CreateBuffer(FRDGBufferDesc::CreateBufferDesc(sizeof(uint32), NumTotalBricks), TEXT("BrickRequestKeys"));
+		FRDGBufferRef BrickRequestSortedIndices = GraphBuilder.CreateBuffer(FRDGBufferDesc::CreateBufferDesc(sizeof(uint32), NumTotalBricks), TEXT("BrickRequestSortedIndices"));
+		FRDGBufferRef BrickRequestKeys2 = GraphBuilder.CreateBuffer(FRDGBufferDesc::CreateBufferDesc(sizeof(uint32), NumTotalBricks), TEXT("BrickRequestKeys2"));
+		FRDGBufferRef BrickRequestSortedIndices2 = GraphBuilder.CreateBuffer(FRDGBufferDesc::CreateBufferDesc(sizeof(uint32), NumTotalBricks), TEXT("BrickRequestSortedIndices2"));
+		FRDGBufferRef BrickRequestSortedIndicesInverse = GraphBuilder.CreateBuffer(FRDGBufferDesc::CreateBufferDesc(sizeof(uint32), NumTotalBricks), TEXT("BrickRequestSortedIndices2"));
 
 		TArray<FRDGTextureUAV*, FRDGArrayAllocator> VoxelizationVolumeMipUAVs;
 
@@ -462,9 +497,12 @@ void FVolumetricLightmapRenderer::VoxelizeScene()
 			FGatherBrickRequestsCS::FParameters* PassParameters = GraphBuilder.AllocParameters<FGatherBrickRequestsCS::FParameters>();
 			PassParameters->VolumeSize = VoxelizationVolumeMips[MipLevel]->GetDesc().GetSize();
 			PassParameters->BrickSize = 1 << (MipLevel * BrickSizeLog2);
+			PassParameters->MipLevel = MipLevel;
 			PassParameters->VoxelizeVolume = VoxelizationVolumeMipUAVs[MipLevel];
 			PassParameters->BrickAllocatorParameters = BrickAllocatorParameters.UAV;
-			PassParameters->BrickRequests = BrickRequests.UAV;
+			PassParameters->BrickRequestsUnsorted = GraphBuilder.CreateUAV(BrickRequestsUnsorted, PF_R32G32B32A32_UINT);
+			PassParameters->BrickRequestKeys = GraphBuilder.CreateUAV(BrickRequestKeys, PF_R32_UINT);
+			PassParameters->BrickRequestSortedIndices = GraphBuilder.CreateUAV(BrickRequestSortedIndices, PF_R32_UINT);
 
 			FComputeShaderUtils::AddPass(
 				GraphBuilder,
@@ -472,10 +510,81 @@ void FVolumetricLightmapRenderer::VoxelizeScene()
 				ComputeShader,
 				PassParameters,
 				FComputeShaderUtils::GetGroupCount(VoxelizationVolumeMips[MipLevel]->GetDesc().GetSize(), FIntVector(4)));
+		}
+		
+		{
+			FSortBrickRequestsPassParameters* PassParameters = GraphBuilder.AllocParameters<FSortBrickRequestsPassParameters>();
+			PassParameters->BrickRequestKeysSRV = GraphBuilder.CreateSRV(BrickRequestKeys, PF_R32_UINT);
+			PassParameters->BrickRequestKeysSRV2 = GraphBuilder.CreateSRV(BrickRequestKeys2, PF_R32_UINT);
+			PassParameters->BrickRequestKeysUAV = GraphBuilder.CreateUAV(BrickRequestKeys, PF_R32_UINT);
+			PassParameters->BrickRequestKeysUAV2 = GraphBuilder.CreateUAV(BrickRequestKeys2, PF_R32_UINT);
+			PassParameters->BrickRequestSortedIndicesSRV = GraphBuilder.CreateSRV(BrickRequestSortedIndices, PF_R32_UINT);
+			PassParameters->BrickRequestSortedIndicesSRV2 = GraphBuilder.CreateSRV(BrickRequestSortedIndices2, PF_R32_UINT);
+			PassParameters->BrickRequestSortedIndicesUAV = GraphBuilder.CreateUAV(BrickRequestSortedIndices, PF_R32_UINT);
+			PassParameters->BrickRequestSortedIndicesUAV2 = GraphBuilder.CreateUAV(BrickRequestSortedIndices2, PF_R32_UINT);
+
+			AddClearUAVPass(GraphBuilder, GraphBuilder.CreateUAV(BrickRequestKeys2, PF_R32_UINT), 0);
+			AddClearUAVPass(GraphBuilder, GraphBuilder.CreateUAV(BrickRequestSortedIndices2, PF_R32_UINT), 0);
+			
+			GraphBuilder.AddPass(
+				RDG_EVENT_NAME("SortBrickRequests"),
+				PassParameters,
+				ERDGPassFlags::Compute,
+				[
+					PassParameters,
+					NumTotalBricks = NumTotalBricks
+				](FRHIRayTracingCommandList& RHICmdList)
+				{
+					FGPUSortBuffers SortBuffers;
+					SortBuffers.RemoteKeySRVs[0] = PassParameters->BrickRequestKeysSRV->GetRHI();
+					SortBuffers.RemoteKeySRVs[1] = PassParameters->BrickRequestKeysSRV2->GetRHI();
+					SortBuffers.RemoteKeyUAVs[0] = PassParameters->BrickRequestKeysUAV->GetRHI();
+					SortBuffers.RemoteKeyUAVs[1] = PassParameters->BrickRequestKeysUAV2->GetRHI();
+					SortBuffers.RemoteValueSRVs[0] = PassParameters->BrickRequestSortedIndicesSRV->GetRHI();
+					SortBuffers.RemoteValueSRVs[1] = PassParameters->BrickRequestSortedIndicesSRV2->GetRHI();
+					SortBuffers.RemoteValueUAVs[0] = PassParameters->BrickRequestSortedIndicesUAV->GetRHI();
+					SortBuffers.RemoteValueUAVs[1] = PassParameters->BrickRequestSortedIndicesUAV2->GetRHI();
+
+					int32 ResultBufferIndex = SortGPUBuffers(RHICmdList, SortBuffers, 0, 0xFFFFFFFF, NumTotalBricks, GMaxRHIFeatureLevel);
+					check(ResultBufferIndex == 0);
+				}
+			);
+		}
 
-			RHICmdList.Transition(FRHITransitionInfo(BrickRequests.UAV, ERHIAccess::Unknown, ERHIAccess::UAVCompute));
+		{
+			FWriteSortedBrickRequestsCS::FParameters* PassParameters = GraphBuilder.AllocParameters<FWriteSortedBrickRequestsCS::FParameters>();
+			PassParameters->BrickRequests = BrickRequests.UAV;
+			PassParameters->BrickRequestsUnsorted = GraphBuilder.CreateUAV(BrickRequestsUnsorted, PF_R32G32B32A32_UINT);
+			PassParameters->BrickRequestSortedIndices = GraphBuilder.CreateUAV(BrickRequestSortedIndices, PF_R32_UINT);
+			PassParameters->BrickRequestSortedIndicesInverse = GraphBuilder.CreateUAV(BrickRequestSortedIndicesInverse, PF_R32_UINT);
+
+			TShaderMapRef<FWriteSortedBrickRequestsCS> ComputeShader(GlobalShaderMap);
+			
+			FComputeShaderUtils::AddPass(
+				GraphBuilder,
+				RDG_EVENT_NAME("WriteSortedBrickRequests"),
+				ComputeShader,
+				PassParameters,
+				FIntVector(NumTotalBricks, 1, 1));
 		}
 
+		for (int32 MipLevel = VoxelizationVolumeMips.Num() - 1; MipLevel >= 0; MipLevel--)
+		{
+			TShaderMapRef<FPermuteVoxelizeVolumeCS> ComputeShader(GlobalShaderMap);
+
+			FPermuteVoxelizeVolumeCS::FParameters* PassParameters = GraphBuilder.AllocParameters<FPermuteVoxelizeVolumeCS::FParameters>();
+			PassParameters->VolumeSize = VoxelizationVolumeMips[MipLevel]->GetDesc().GetSize();
+			PassParameters->VoxelizeVolume = VoxelizationVolumeMipUAVs[MipLevel];
+			PassParameters->BrickRequestSortedIndicesInverse = GraphBuilder.CreateUAV(BrickRequestSortedIndicesInverse, PF_R32_UINT);
+
+			FComputeShaderUtils::AddPass(
+				GraphBuilder,
+				RDG_EVENT_NAME("PermuteVoxelizeVolume"),
+				ComputeShader,
+				PassParameters,
+				FComputeShaderUtils::GetGroupCount(VoxelizationVolumeMips[MipLevel]->GetDesc().GetSize(), FIntVector(4)));
+		}
+		
 		for (int32 MipLevel = VoxelizationVolumeMips.Num() - 1; MipLevel >= 0; MipLevel--)
 		{
 			TShaderMapRef<FSplatVolumeCS> ComputeShader(GlobalShaderMap);
@@ -502,6 +611,103 @@ void FVolumetricLightmapRenderer::VoxelizeScene()
 	Scene->DestroyRayTracingScene();
 }
 
+struct FVolumetricLightmapBrickDataRDG
+{
+	FRDGTextureUAVRef AmbientVector;
+	FRDGTextureUAVRef SHCoefficients[6];
+	FRDGTextureUAVRef SkyBentNormal;
+	FRDGTextureUAVRef DirectionalLightShadowing;
+
+	TRefCountPtr<IPooledRenderTarget> AmbientVectorWrapperRT;
+	TRefCountPtr<IPooledRenderTarget> SHCoefficientsWrapperRT[6];
+	TRefCountPtr<IPooledRenderTarget> SkyBentNormalWrapperRT;
+	TRefCountPtr<IPooledRenderTarget> DirectionalLightShadowingWrapperRT;
+
+	FRDGTextureRef AmbientVectorRDG;
+	FRDGTextureRef SHCoefficientsRDG[6];
+	FRDGTextureRef SkyBentNormalRDG;
+	FRDGTextureRef DirectionalLightShadowingRDG;
+
+	TRefCountPtr<IPooledRenderTarget> CreateWrapperRenderTargetForLayer(FVolumetricLightmapDataLayer& Layer, FIntVector Dimensions)
+	{
+		return TRefCountPtr<IPooledRenderTarget>(new FPooledRenderTarget(
+				Layer.Texture,
+				FPooledRenderTargetDesc::CreateVolumeDesc(
+					Dimensions.X, Dimensions.Y, Dimensions.Z,
+					Layer.Format,
+					FClearValueBinding::None,
+					ETextureCreateFlags::ShaderResource | ETextureCreateFlags::UAV,
+					ETextureCreateFlags::ShaderResource | ETextureCreateFlags::UAV,
+					false),
+				nullptr
+			)
+		);
+	}
+	
+	void CreateFromBrickData(FRDGBuilder& GraphBuilder, FVolumetricLightmapBrickData& BrickData, FIntVector Dimensions)
+	{
+		AmbientVectorWrapperRT = CreateWrapperRenderTargetForLayer(BrickData.AmbientVector, Dimensions);
+		AmbientVectorRDG = GraphBuilder.RegisterExternalTexture(AmbientVectorWrapperRT);
+		AmbientVector = GraphBuilder.CreateUAV(AmbientVectorRDG);
+
+		for (int32 CoefficientIndex = 0; CoefficientIndex < UE_ARRAY_COUNT(SHCoefficients); CoefficientIndex++)
+		{
+			SHCoefficientsWrapperRT[CoefficientIndex] = CreateWrapperRenderTargetForLayer(BrickData.SHCoefficients[CoefficientIndex], Dimensions);
+			SHCoefficientsRDG[CoefficientIndex] = GraphBuilder.RegisterExternalTexture(SHCoefficientsWrapperRT[CoefficientIndex]);
+			SHCoefficients[CoefficientIndex] = GraphBuilder.CreateUAV(SHCoefficientsRDG[CoefficientIndex]);
+		}
+
+		SkyBentNormalWrapperRT = CreateWrapperRenderTargetForLayer(BrickData.SkyBentNormal, Dimensions);
+		SkyBentNormalRDG = GraphBuilder.RegisterExternalTexture(SkyBentNormalWrapperRT);
+		SkyBentNormal = GraphBuilder.CreateUAV(SkyBentNormalRDG);
+		
+		DirectionalLightShadowingWrapperRT = CreateWrapperRenderTargetForLayer(BrickData.DirectionalLightShadowing, Dimensions);
+		DirectionalLightShadowingRDG = GraphBuilder.RegisterExternalTexture(DirectionalLightShadowingWrapperRT);
+		DirectionalLightShadowing = GraphBuilder.CreateUAV(DirectionalLightShadowingRDG);		
+	}
+
+	void QueueTextureExtractions(FRDGBuilder& GraphBuilder)
+	{
+		GraphBuilder.QueueTextureExtraction(AmbientVectorRDG, &AmbientVectorWrapperRT);
+
+		for (int32 CoefficientIndex = 0; CoefficientIndex < UE_ARRAY_COUNT(SHCoefficients); CoefficientIndex++)
+		{
+			GraphBuilder.QueueTextureExtraction(SHCoefficientsRDG[CoefficientIndex], &SHCoefficientsWrapperRT[CoefficientIndex]);
+		}
+
+		GraphBuilder.QueueTextureExtraction(SkyBentNormalRDG, &SkyBentNormalWrapperRT);		
+		GraphBuilder.QueueTextureExtraction(DirectionalLightShadowingRDG, &DirectionalLightShadowingWrapperRT);
+	}
+
+	TRefCountPtr<IPooledRenderTarget> OldAmbientVectorWrapperRT;
+	TRefCountPtr<IPooledRenderTarget> OldSHCoefficientsWrapperRT[6];
+	TRefCountPtr<IPooledRenderTarget> OldSkyBentNormalWrapperRT;
+	TRefCountPtr<IPooledRenderTarget> OldDirectionalLightShadowingWrapperRT;
+
+	void DebugSetupInvariants()
+	{
+		OldAmbientVectorWrapperRT = AmbientVectorWrapperRT;
+		for (int32 CoefficientIndex = 0; CoefficientIndex < UE_ARRAY_COUNT(SHCoefficients); CoefficientIndex++)
+		{
+			OldSHCoefficientsWrapperRT[CoefficientIndex] = SHCoefficientsWrapperRT[CoefficientIndex];
+		}
+		OldSkyBentNormalWrapperRT = SkyBentNormalWrapperRT;
+		OldDirectionalLightShadowingWrapperRT = DirectionalLightShadowingWrapperRT;
+	}
+
+	void DebugCheckInvariants()
+	{
+		checkf(OldAmbientVectorWrapperRT == AmbientVectorWrapperRT,
+			TEXT("The code relies on the fact that the texture extracted from RDG is exactly the same as the one registered"));
+		for (int32 CoefficientIndex = 0; CoefficientIndex < UE_ARRAY_COUNT(SHCoefficients); CoefficientIndex++)
+		{
+			check(OldSHCoefficientsWrapperRT[CoefficientIndex] == SHCoefficientsWrapperRT[CoefficientIndex]);
+		}
+		check(OldSkyBentNormalWrapperRT == SkyBentNormalWrapperRT);
+		check(OldDirectionalLightShadowingWrapperRT == DirectionalLightShadowingWrapperRT);
+	}
+};
+	
 void FVolumetricLightmapRenderer::BackgroundTick()
 {
 	if (NumTotalBricks == 0)
@@ -509,8 +715,7 @@ void FVolumetricLightmapRenderer::BackgroundTick()
 		return;
 	}
 
-	int32 NumCellsPerBrick = 5 * 5 * 5;
-	if (SamplesTaken >= (uint64)NumTotalBricks * NumCellsPerBrick * NumTotalPassesToRender)
+	if (SamplesTaken >= (uint64)NumTotalBricks * NumTotalPassesToRender)
 	{
 		return;
 	}
@@ -527,271 +732,218 @@ void FVolumetricLightmapRenderer::BackgroundTick()
 
 	FRDGBuilder GraphBuilder(FRHICommandListExecutor::GetImmediateCommandList());	
 
+	FVolumetricLightmapBrickDataRDG AccumulationBrickDataRDG;
+	AccumulationBrickDataRDG.CreateFromBrickData(GraphBuilder, AccumulationBrickData, VolumetricLightmapData.BrickDataDimensions);
+	AccumulationBrickDataRDG.DebugSetupInvariants();
+		
+	FVolumetricLightmapBrickDataRDG OutputBrickDataRDG;
+	OutputBrickDataRDG.CreateFromBrickData(GraphBuilder, VolumetricLightmapData.BrickData, VolumetricLightmapData.BrickDataDimensions);	
+	OutputBrickDataRDG.DebugSetupInvariants();
+	
 	RectLightAtlas::UpdateRectLightAtlasTexture(GraphBuilder, Scene->FeatureLevel);
+
+	FGlobalShaderMap* GlobalShaderMap = GetGlobalShaderMap(Scene->FeatureLevel);
+
+	const bool bIsViewportNonRealtime = GCurrentLevelEditingViewportClient && !GCurrentLevelEditingViewportClient->IsRealtime();
+	const int32 NumSamplesPerFrame = 8 * (bIsViewportNonRealtime ? Scene->Settings->TilePassesInFullSpeedMode : Scene->Settings->TilePassesInSlowMode);
+
+	FVolumetricLightmapPathTracingRGS::FParameters* PreviousPassParameters = nullptr;
 	
+	for (int32 SampleIndex = 0; SampleIndex < NumSamplesPerFrame; SampleIndex++)
 	{
-		RDG_EVENT_SCOPE(GraphBuilder, "Volumetric Lightmap Path Tracing");
-
-		FGlobalShaderMap* GlobalShaderMap = GetGlobalShaderMap(Scene->FeatureLevel);
-
-		bool bLastFewFramesIdle = !GCurrentLevelEditingViewportClient || !GCurrentLevelEditingViewportClient->IsRealtime();
-
-		int32 NumSamplesThisFrame = !bLastFewFramesIdle ? 1 : 32;
-
-		// manually handle transitions since the buffers are not (yet) managed by RDG
-		GraphBuilder.AddPass(RDG_EVENT_NAME("Transition Buffers"), ERDGPassFlags::None,
-			[this](FRHICommandList& RHICmdList) {
-			RHICmdList.Transition({
-				FRHITransitionInfo(VolumetricLightmapData.BrickData.AmbientVector.UAV, ERHIAccess::Unknown, ERHIAccess::UAVCompute),
-				FRHITransitionInfo(VolumetricLightmapData.BrickData.SHCoefficients[0].UAV, ERHIAccess::Unknown, ERHIAccess::UAVCompute),
-				FRHITransitionInfo(VolumetricLightmapData.BrickData.SHCoefficients[1].UAV, ERHIAccess::Unknown, ERHIAccess::UAVCompute),
-				FRHITransitionInfo(VolumetricLightmapData.BrickData.SHCoefficients[2].UAV, ERHIAccess::Unknown, ERHIAccess::UAVCompute),
-				FRHITransitionInfo(VolumetricLightmapData.BrickData.SHCoefficients[3].UAV, ERHIAccess::Unknown, ERHIAccess::UAVCompute),
-				FRHITransitionInfo(VolumetricLightmapData.BrickData.SHCoefficients[4].UAV, ERHIAccess::Unknown, ERHIAccess::UAVCompute),
-				FRHITransitionInfo(VolumetricLightmapData.BrickData.SHCoefficients[5].UAV, ERHIAccess::Unknown, ERHIAccess::UAVCompute),
-				FRHITransitionInfo(VolumetricLightmapData.BrickData.SkyBentNormal.UAV, ERHIAccess::Unknown, ERHIAccess::UAVCompute),
-				FRHITransitionInfo(VolumetricLightmapData.BrickData.DirectionalLightShadowing.UAV, ERHIAccess::Unknown, ERHIAccess::UAVCompute),
-				}
-			);
-		}
-		);
-
-		FVolumetricLightmapPathTracingRGS::FParameters* PreviousPassParameters = nullptr;
-		for (int32 SampleIndex = 0; SampleIndex < NumSamplesThisFrame; SampleIndex++)
-		{
-			int32 MaxBricksPerFrame = FMath::Min(512, NumTotalBricks);
-			int32 NumFramesOneRound = FMath::DivideAndRoundUp(NumTotalBricks, MaxBricksPerFrame);
-			int32 BrickBatchOffset = MaxBricksPerFrame * (FrameNumber % NumFramesOneRound);
-			int32 BricksToCalcThisFrame = FMath::Min(MaxBricksPerFrame, NumTotalBricks - BrickBatchOffset);
-			if (BricksToCalcThisFrame <= 0) continue;
+		const uint64 NumTotalSamples = (uint64)NumTotalBricks * NumTotalPassesToRender;
+		const int32 BrickBatchIndexToWorkOn = SamplesTaken / (NumTotalPassesToRender * BrickBatchSize);
+		const int32 BrickBatchOffset = BrickBatchIndexToWorkOn * BrickBatchSize;
+		const int32 BricksToCalcThisFrame = FMath::Min(BrickBatchSize, NumTotalBricks - BrickBatchOffset);
+		const int32 RenderPassIndex = (SamplesTaken - BrickBatchOffset * NumTotalPassesToRender) / BricksToCalcThisFrame;
+		if (BricksToCalcThisFrame <= 0) continue;
 
 #if RHI_RAYTRACING
-			if (IsRayTracingEnabled())
+		if (IsRayTracingEnabled())
+		{
+			FVolumetricLightmapPathTracingRGS::FPermutationDomain PermutationVector;
+			PermutationVector.Set<FVolumetricLightmapPathTracingRGS::FUseIrradianceCaching>(Scene->Settings->bUseIrradianceCaching);
+			TShaderRef<FVolumetricLightmapPathTracingRGS> RayGenShader = GlobalShaderMap->GetShader<FVolumetricLightmapPathTracingRGS>(PermutationVector);
+
+			FVolumetricLightmapPathTracingRGS::FParameters* PassParameters = GraphBuilder.AllocParameters<FVolumetricLightmapPathTracingRGS::FParameters>();
+			CA_ASSUME(PassParameters);
+
+			PassParameters->FrameNumber = RenderPassIndex;
+			PassParameters->VolumeMin = (FVector3f)VolumeMin; // LWC_TODO: precision loss
+			PassParameters->VolumeSize = (FVector3f)VolumeSize; // LWC_TODO: precision loss
+			PassParameters->IndirectionTextureDim = IndirectionTextureDimensions;
+			PassParameters->TLAS = Scene->RayTracingSceneSRV;
+			PassParameters->BrickRequests = BrickRequests.SRV;
+			PassParameters->NumTotalBricks = NumTotalBricks;
+			PassParameters->BrickBatchOffset = BrickBatchOffset;
+			PassParameters->AmbientVector = AccumulationBrickDataRDG.AmbientVector;
+			PassParameters->SHCoefficients0R = AccumulationBrickDataRDG.SHCoefficients[0];
+			PassParameters->SHCoefficients1R = AccumulationBrickDataRDG.SHCoefficients[1];
+			PassParameters->SHCoefficients0G = AccumulationBrickDataRDG.SHCoefficients[2];
+			PassParameters->SHCoefficients1G = AccumulationBrickDataRDG.SHCoefficients[3];
+			PassParameters->SHCoefficients0B = AccumulationBrickDataRDG.SHCoefficients[4];
+			PassParameters->SHCoefficients1B = AccumulationBrickDataRDG.SHCoefficients[5];
+			PassParameters->SkyBentNormal = AccumulationBrickDataRDG.SkyBentNormal;
+			PassParameters->DirectionalLightShadowing = AccumulationBrickDataRDG.DirectionalLightShadowing;
+			PassParameters->ViewUniformBuffer = Scene->ReferenceView->ViewUniformBuffer;
+			PassParameters->IrradianceCachingParameters = Scene->IrradianceCache->IrradianceCachingParametersUniformBuffer;
+
+			if (PreviousPassParameters == nullptr)
 			{
-				{
-					FVolumetricLightmapPathTracingRGS::FPermutationDomain PermutationVector;
-					PermutationVector.Set<FVolumetricLightmapPathTracingRGS::FUseIrradianceCaching>(Scene->Settings->bUseIrradianceCaching);
-					TShaderRef<FVolumetricLightmapPathTracingRGS> RayGenShader = GlobalShaderMap->GetShader<FVolumetricLightmapPathTracingRGS>(PermutationVector);
-
-					FVolumetricLightmapPathTracingRGS::FParameters* PassParameters = GraphBuilder.AllocParameters<FVolumetricLightmapPathTracingRGS::FParameters>();
-					CA_ASSUME(PassParameters);
-
-					PassParameters->FrameNumber = FrameNumber / NumFramesOneRound;
-					PassParameters->VolumeMin = (FVector3f)VolumeMin; // LWC_TODO: precision loss
-					PassParameters->VolumeSize = (FVector3f)VolumeSize; // LWC_TODO: precision loss
-					PassParameters->IndirectionTextureDim = IndirectionTextureDimensions;
-					PassParameters->TLAS = Scene->RayTracingSceneSRV;
-					PassParameters->BrickRequests = BrickRequests.SRV;
-					PassParameters->NumTotalBricks = NumTotalBricks;
-					PassParameters->BrickBatchOffset = BrickBatchOffset;
-					PassParameters->AmbientVector = AccumulationBrickData.AmbientVector.UAV;
-					PassParameters->SHCoefficients0R = AccumulationBrickData.SHCoefficients[0].UAV;
-					PassParameters->SHCoefficients1R = AccumulationBrickData.SHCoefficients[1].UAV;
-					PassParameters->SHCoefficients0G = AccumulationBrickData.SHCoefficients[2].UAV;
-					PassParameters->SHCoefficients1G = AccumulationBrickData.SHCoefficients[3].UAV;
-					PassParameters->SHCoefficients0B = AccumulationBrickData.SHCoefficients[4].UAV;
-					PassParameters->SHCoefficients1B = AccumulationBrickData.SHCoefficients[5].UAV;
-					PassParameters->SkyBentNormal = AccumulationBrickData.SkyBentNormal.UAV;
-					PassParameters->DirectionalLightShadowing = AccumulationBrickData.DirectionalLightShadowing.UAV;
-					PassParameters->ViewUniformBuffer = Scene->ReferenceView->ViewUniformBuffer;
-					PassParameters->IrradianceCachingParameters = Scene->IrradianceCache->IrradianceCachingParametersUniformBuffer;
-
-					if (PreviousPassParameters == nullptr)
-					{
-						SetupPathTracingLightParameters(Scene->LightSceneRenderState, GraphBuilder, *Scene->ReferenceView, PassParameters);
-						PreviousPassParameters = PassParameters;
-					}
-					else
-					{
-						PassParameters->LightGridParameters = PreviousPassParameters->LightGridParameters;
-						PassParameters->SceneLightCount = PreviousPassParameters->SceneLightCount;
-						PassParameters->SceneVisibleLightCount = PreviousPassParameters->SceneVisibleLightCount;
-						PassParameters->SceneLights = PreviousPassParameters->SceneLights;
-						PassParameters->SkylightTexture = PreviousPassParameters->SkylightTexture;
-						PassParameters->SkylightTextureSampler = PreviousPassParameters->SkylightTextureSampler;
-						PassParameters->SkylightPdf = PreviousPassParameters->SkylightPdf;
-						PassParameters->SkylightInvResolution = PreviousPassParameters->SkylightInvResolution;
-						PassParameters->SkylightMipCount = PreviousPassParameters->SkylightMipCount;
-						PassParameters->IESTexture = PreviousPassParameters->IESTexture;
-						PassParameters->IESTextureSampler = PreviousPassParameters->IESTextureSampler;
-					}
-
-					PassParameters->SSProfilesTexture = GetSubsurfaceProfileTexture();
-
-					TArray<FLightShaderConstants> OptionalStationaryDirectionalLightShadowing;
-					for (FDirectionalLightRenderState& DirectionalLight : Scene->LightSceneRenderState.DirectionalLights.Elements)
-					{
-						if (DirectionalLight.bStationary)
-						{
-							OptionalStationaryDirectionalLightShadowing.Add(DirectionalLight.GetLightShaderParameters());
-							break;
-						}
-					}
-					if (OptionalStationaryDirectionalLightShadowing.Num() == 0)
-					{
-						OptionalStationaryDirectionalLightShadowing.AddZeroed();
-					}
-					PassParameters->LightShaderParametersArray = GraphBuilder.CreateSRV(FRDGBufferSRVDesc(CreateStructuredBuffer(GraphBuilder, TEXT("OptionalStationaryDirectionalLightShadowing"), sizeof(FLightShaderConstants),
-						OptionalStationaryDirectionalLightShadowing.Num(), 
-						OptionalStationaryDirectionalLightShadowing.GetData(),
-						sizeof(FLightShaderConstants) * OptionalStationaryDirectionalLightShadowing.Num()
-					)));
-
-					FSceneRenderState* SceneRenderState = Scene; // capture member variable
-					GraphBuilder.AddPass(
-						RDG_EVENT_NAME("VolumetricLightmapPathTracing %d bricks %d rays", BricksToCalcThisFrame, BricksToCalcThisFrame * (BrickSize + 1) * (BrickSize + 1) * (BrickSize + 1)),
-						PassParameters,
-						ERDGPassFlags::Compute,
-						[PassParameters, RayGenShader, SceneRenderState, BricksToCalcThisFrame](FRHIRayTracingCommandList& RHICmdList)
-					{
-						FRayTracingShaderBindingsWriter GlobalResources;
-						SetShaderParameters(GlobalResources, RayGenShader, *PassParameters);
-
-						FRHIRayTracingScene* RayTracingSceneRHI = SceneRenderState->RayTracingScene;
-						RHICmdList.RayTraceDispatch(SceneRenderState->RayTracingPipelineState, RayGenShader.GetRayTracingShader(), RayTracingSceneRHI, GlobalResources, BricksToCalcThisFrame * (BrickSize + 1) * (BrickSize + 1) * (BrickSize + 1), 1);
-					}
-					);
+				SetupPathTracingLightParameters(Scene->LightSceneRenderState, GraphBuilder, *Scene->ReferenceView, PassParameters);
+				PreviousPassParameters = PassParameters;
+			}
+			else
+			{
+				PassParameters->LightGridParameters = PreviousPassParameters->LightGridParameters;
+				PassParameters->SceneLightCount = PreviousPassParameters->SceneLightCount;
+				PassParameters->SceneVisibleLightCount = PreviousPassParameters->SceneVisibleLightCount;
+				PassParameters->SceneLights = PreviousPassParameters->SceneLights;
+				PassParameters->SkylightTexture = PreviousPassParameters->SkylightTexture;
+				PassParameters->SkylightTextureSampler = PreviousPassParameters->SkylightTextureSampler;
+				PassParameters->SkylightPdf = PreviousPassParameters->SkylightPdf;
+				PassParameters->SkylightInvResolution = PreviousPassParameters->SkylightInvResolution;
+				PassParameters->SkylightMipCount = PreviousPassParameters->SkylightMipCount;
+				PassParameters->IESTexture = PreviousPassParameters->IESTexture;
+				PassParameters->IESTextureSampler = PreviousPassParameters->IESTextureSampler;
+			}
 
+			PassParameters->SSProfilesTexture = GetSubsurfaceProfileTexture();
 
+			TArray<FLightShaderConstants> OptionalStationaryDirectionalLightShadowing;
+			for (FDirectionalLightRenderState& DirectionalLight : Scene->LightSceneRenderState.DirectionalLights.Elements)
+			{
+				if (DirectionalLight.bStationary)
+				{
+					OptionalStationaryDirectionalLightShadowing.Add(DirectionalLight.GetLightShaderParameters());
+					break;
 				}
 			}
-#endif
+			
+			if (OptionalStationaryDirectionalLightShadowing.Num() == 0)
 			{
-				// manually handle transitions since the buffers are not (yet) managed by RDG
-				GraphBuilder.AddPass(RDG_EVENT_NAME("Transition Buffers"), ERDGPassFlags::None,
-					[this](FRHICommandList& RHICmdList) {
-					RHICmdList.Transition({
-						FRHITransitionInfo(AccumulationBrickData.AmbientVector.UAV, ERHIAccess::UAVCompute, ERHIAccess::UAVCompute),
-						FRHITransitionInfo(AccumulationBrickData.SHCoefficients[0].UAV, ERHIAccess::UAVCompute, ERHIAccess::UAVCompute),
-						FRHITransitionInfo(AccumulationBrickData.SHCoefficients[1].UAV, ERHIAccess::UAVCompute, ERHIAccess::UAVCompute),
-						FRHITransitionInfo(AccumulationBrickData.SHCoefficients[2].UAV, ERHIAccess::UAVCompute, ERHIAccess::UAVCompute),
-						FRHITransitionInfo(AccumulationBrickData.SHCoefficients[3].UAV, ERHIAccess::UAVCompute, ERHIAccess::UAVCompute),
-						FRHITransitionInfo(AccumulationBrickData.SHCoefficients[4].UAV, ERHIAccess::UAVCompute, ERHIAccess::UAVCompute),
-						FRHITransitionInfo(AccumulationBrickData.SHCoefficients[5].UAV, ERHIAccess::UAVCompute, ERHIAccess::UAVCompute),
-						FRHITransitionInfo(AccumulationBrickData.SkyBentNormal.UAV, ERHIAccess::UAVCompute, ERHIAccess::UAVCompute),
-						FRHITransitionInfo(AccumulationBrickData.DirectionalLightShadowing.UAV, ERHIAccess::UAVCompute, ERHIAccess::UAVCompute),
-						}
-					);
-				}
-				);
-
-				TShaderMapRef<FFinalizeBrickResultsCS> ComputeShader(GlobalShaderMap);
-
-				FFinalizeBrickResultsCS::FParameters* PassParameters = GraphBuilder.AllocParameters<FFinalizeBrickResultsCS::FParameters>();
-				PassParameters->NumTotalBricks = NumTotalBricks;
-				PassParameters->BrickBatchOffset = BrickBatchOffset;
-				PassParameters->BrickRequests = BrickRequests.UAV;
-				PassParameters->AmbientVector = AccumulationBrickData.AmbientVector.Texture;
-				PassParameters->SHCoefficients0R = AccumulationBrickData.SHCoefficients[0].Texture;
-				PassParameters->SHCoefficients1R = AccumulationBrickData.SHCoefficients[1].Texture;
-				PassParameters->SHCoefficients0G = AccumulationBrickData.SHCoefficients[2].Texture;
-				PassParameters->SHCoefficients1G = AccumulationBrickData.SHCoefficients[3].Texture;
-				PassParameters->SHCoefficients0B = AccumulationBrickData.SHCoefficients[4].Texture;
-				PassParameters->SHCoefficients1B = AccumulationBrickData.SHCoefficients[5].Texture;
-				PassParameters->SkyBentNormal = AccumulationBrickData.SkyBentNormal.Texture;
-				PassParameters->DirectionalLightShadowing = AccumulationBrickData.DirectionalLightShadowing.Texture;
-				PassParameters->OutAmbientVector = VolumetricLightmapData.BrickData.AmbientVector.UAV;
-				PassParameters->OutSHCoefficients0R = VolumetricLightmapData.BrickData.SHCoefficients[0].UAV;
-				PassParameters->OutSHCoefficients1R = VolumetricLightmapData.BrickData.SHCoefficients[1].UAV;
-				PassParameters->OutSHCoefficients0G = VolumetricLightmapData.BrickData.SHCoefficients[2].UAV;
-				PassParameters->OutSHCoefficients1G = VolumetricLightmapData.BrickData.SHCoefficients[3].UAV;
-				PassParameters->OutSHCoefficients0B = VolumetricLightmapData.BrickData.SHCoefficients[4].UAV;
-				PassParameters->OutSHCoefficients1B = VolumetricLightmapData.BrickData.SHCoefficients[5].UAV;
-				PassParameters->OutSkyBentNormal = VolumetricLightmapData.BrickData.SkyBentNormal.UAV;
-				PassParameters->OutDirectionalLightShadowing = VolumetricLightmapData.BrickData.DirectionalLightShadowing.UAV;
-
-				FComputeShaderUtils::AddPass(GraphBuilder, RDG_EVENT_NAME("FinalizeBrickResults"), ComputeShader, PassParameters, FIntVector(BricksToCalcThisFrame, 1, 1));
-
-				// manually handle transitions since the buffers are not (yet) managed by RDG
-				GraphBuilder.AddPass(RDG_EVENT_NAME("Transition Buffers"), ERDGPassFlags::None,
-					[this](FRHICommandList& RHICmdList) {
-					RHICmdList.Transition({
-						FRHITransitionInfo(AccumulationBrickData.AmbientVector.UAV, ERHIAccess::UAVCompute, ERHIAccess::UAVCompute),
-						FRHITransitionInfo(AccumulationBrickData.SHCoefficients[0].UAV, ERHIAccess::UAVCompute, ERHIAccess::UAVCompute),
-						FRHITransitionInfo(AccumulationBrickData.SHCoefficients[1].UAV, ERHIAccess::UAVCompute, ERHIAccess::UAVCompute),
-						FRHITransitionInfo(AccumulationBrickData.SHCoefficients[2].UAV, ERHIAccess::UAVCompute, ERHIAccess::UAVCompute),
-						FRHITransitionInfo(AccumulationBrickData.SHCoefficients[3].UAV, ERHIAccess::UAVCompute, ERHIAccess::UAVCompute),
-						FRHITransitionInfo(AccumulationBrickData.SHCoefficients[4].UAV, ERHIAccess::UAVCompute, ERHIAccess::UAVCompute),
-						FRHITransitionInfo(AccumulationBrickData.SHCoefficients[5].UAV, ERHIAccess::UAVCompute, ERHIAccess::UAVCompute),
-						FRHITransitionInfo(AccumulationBrickData.SkyBentNormal.UAV, ERHIAccess::UAVCompute, ERHIAccess::UAVCompute),
-						FRHITransitionInfo(AccumulationBrickData.DirectionalLightShadowing.UAV, ERHIAccess::UAVCompute, ERHIAccess::UAVCompute),
-						}
-					);
-				}
-				);
+				OptionalStationaryDirectionalLightShadowing.AddZeroed();
 			}
+			
+			PassParameters->LightShaderParametersArray = GraphBuilder.CreateSRV(FRDGBufferSRVDesc(CreateStructuredBuffer(
+				GraphBuilder, TEXT("OptionalStationaryDirectionalLightShadowing"), sizeof(FLightShaderConstants),
+				OptionalStationaryDirectionalLightShadowing.Num(),
+				OptionalStationaryDirectionalLightShadowing.GetData(),
+				sizeof(FLightShaderConstants) * OptionalStationaryDirectionalLightShadowing.Num()
+			)));
+
+			FSceneRenderState* SceneRenderState = Scene; // capture member variable
+			GraphBuilder.AddPass(
+				RDG_EVENT_NAME("VolumetricLightmapPathTracing %d bricks %d rays", BricksToCalcThisFrame, BricksToCalcThisFrame * (BrickSize + 1) * (BrickSize + 1) * (BrickSize + 1)),
+				PassParameters,
+				ERDGPassFlags::Compute,
+				[
+					PassParameters,
+					RayGenShader,
+					RayTracingSceneRHI = Scene->RayTracingScene,
+					RayTracingPipelineState = Scene->RayTracingPipelineState,
+					BricksToCalcThisFrame
+				](FRHIRayTracingCommandList& RHICmdList)
+				{
+					FRayTracingShaderBindingsWriter GlobalResources;
+					SetShaderParameters(GlobalResources, RayGenShader, *PassParameters);
+
+					RHICmdList.RayTraceDispatch(RayTracingPipelineState, RayGenShader.GetRayTracingShader(), RayTracingSceneRHI, GlobalResources,
+					                            BricksToCalcThisFrame * (BrickSize + 1) * (BrickSize + 1) * (BrickSize + 1), 1);
+				}
+			);
+		}
+#endif
+
+		FRDGTextureUAVRef ValidityBrickDataUAV = GraphBuilder.CreateUAV(GraphBuilder.RegisterExternalTexture(ValidityBrickData));
+		
+		// Output to OutputBrickData if this is the last pass for the current batch
+		if (!bIsViewportNonRealtime || RenderPassIndex == NumTotalPassesToRender - 1)
+		{
+			TShaderMapRef<FFinalizeBrickResultsCS> ComputeShader(GlobalShaderMap);
+
+			FFinalizeBrickResultsCS::FParameters* PassParameters = GraphBuilder.AllocParameters<FFinalizeBrickResultsCS::FParameters>();
+			PassParameters->NumTotalBricks = NumTotalBricks;
+			PassParameters->BrickBatchOffset = BrickBatchOffset;
+			PassParameters->NumTotalPassesToRender = NumTotalPassesToRender;
+			PassParameters->BrickRequests = BrickRequests.UAV;
+			PassParameters->AmbientVector = AccumulationBrickDataRDG.AmbientVectorRDG;
+			PassParameters->SHCoefficients0R = AccumulationBrickDataRDG.SHCoefficientsRDG[0];
+			PassParameters->SHCoefficients1R = AccumulationBrickDataRDG.SHCoefficientsRDG[1];
+			PassParameters->SHCoefficients0G = AccumulationBrickDataRDG.SHCoefficientsRDG[2];
+			PassParameters->SHCoefficients1G = AccumulationBrickDataRDG.SHCoefficientsRDG[3];
+			PassParameters->SHCoefficients0B = AccumulationBrickDataRDG.SHCoefficientsRDG[4];
+			PassParameters->SHCoefficients1B = AccumulationBrickDataRDG.SHCoefficientsRDG[5];
+			PassParameters->SkyBentNormal = AccumulationBrickDataRDG.SkyBentNormalRDG;
+			PassParameters->DirectionalLightShadowing = AccumulationBrickDataRDG.DirectionalLightShadowingRDG;
+			PassParameters->OutAmbientVector = OutputBrickDataRDG.AmbientVector;
+			PassParameters->OutSHCoefficients0R = OutputBrickDataRDG.SHCoefficients[0];
+			PassParameters->OutSHCoefficients1R = OutputBrickDataRDG.SHCoefficients[1];
+			PassParameters->OutSHCoefficients0G = OutputBrickDataRDG.SHCoefficients[2];
+			PassParameters->OutSHCoefficients1G = OutputBrickDataRDG.SHCoefficients[3];
+			PassParameters->OutSHCoefficients0B = OutputBrickDataRDG.SHCoefficients[4];
+			PassParameters->OutSHCoefficients1B = OutputBrickDataRDG.SHCoefficients[5];
+			PassParameters->OutSkyBentNormal = OutputBrickDataRDG.SkyBentNormal;
+			PassParameters->OutDirectionalLightShadowing = OutputBrickDataRDG.DirectionalLightShadowing;
+			PassParameters->ValidityMask = ValidityBrickDataUAV;
+
+			FComputeShaderUtils::AddPass(GraphBuilder, RDG_EVENT_NAME("FinalizeBrickResults"), ComputeShader, PassParameters, FIntVector(BricksToCalcThisFrame, 1, 1));
+		}
 
+		// Do stitching over all brickes if this is the very last pass of all bricks
+		if (SamplesTaken < NumTotalSamples && SamplesTaken + BricksToCalcThisFrame >= NumTotalSamples)
+		{
+			FRDGTextureUAV* IndirectionTextureUAV = GraphBuilder.CreateUAV(GraphBuilder.RegisterExternalTexture(IndirectionTexture));
+
+			// Do 2 passes to propagate across 3 refinement levels
+			for (int32 StitchPass = 0; StitchPass < 2; StitchPass++)
 			{
-				FRDGTextureUAV* IndirectionTextureUAV = GraphBuilder.CreateUAV(GraphBuilder.RegisterExternalTexture(IndirectionTexture));
+				const int32 StitchingBrickBatchSize = 256;
 
-				// Doing 2 passes no longer makes sense in an amortized setup
-				// for (int32 StitchPass = 0; StitchPass < 2; StitchPass++)
+				for (int32 StitchingBrickBatchOffset = 0; StitchingBrickBatchOffset < NumTotalBricks; StitchingBrickBatchOffset += StitchingBrickBatchSize)
 				{
 					TShaderMapRef<FStitchBorderCS> ComputeShader(GlobalShaderMap);
 
 					FStitchBorderCS::FParameters* PassParameters = GraphBuilder.AllocParameters<FStitchBorderCS::FParameters>();
 					PassParameters->BrickDataDimensions = VolumetricLightmapData.BrickDataDimensions;
 					PassParameters->IndirectionTextureDim = IndirectionTextureDimensions;
-
-					int32 EffectiveRenderPassIndex = FrameNumber / NumFramesOneRound;
-
-					if (Scene->Settings->bUseIrradianceCaching)
-					{
-						if (EffectiveRenderPassIndex >= Scene->Settings->IrradianceCacheQuality)
-						{
-							EffectiveRenderPassIndex -= Scene->Settings->IrradianceCacheQuality;
-						}
-					}
-					
-					PassParameters->FrameNumber = EffectiveRenderPassIndex;
+					PassParameters->FrameNumber = NumTotalPassesToRender;
 					PassParameters->NumTotalBricks = NumTotalBricks;
-					PassParameters->BrickBatchOffset = BrickBatchOffset;
+					PassParameters->BrickBatchOffset = StitchingBrickBatchOffset;
 					PassParameters->IndirectionTexture = IndirectionTextureUAV;
+					PassParameters->ValidityMask = ValidityBrickDataUAV;
 					PassParameters->BrickRequests = BrickRequests.UAV;
-					PassParameters->AmbientVector = AccumulationBrickData.AmbientVector.Texture;
-					PassParameters->OutAmbientVector = VolumetricLightmapData.BrickData.AmbientVector.UAV;
-					PassParameters->OutSHCoefficients0R = VolumetricLightmapData.BrickData.SHCoefficients[0].UAV;
-					PassParameters->OutSHCoefficients1R = VolumetricLightmapData.BrickData.SHCoefficients[1].UAV;
-					PassParameters->OutSHCoefficients0G = VolumetricLightmapData.BrickData.SHCoefficients[2].UAV;
-					PassParameters->OutSHCoefficients1G = VolumetricLightmapData.BrickData.SHCoefficients[3].UAV;
-					PassParameters->OutSHCoefficients0B = VolumetricLightmapData.BrickData.SHCoefficients[4].UAV;
-					PassParameters->OutSHCoefficients1B = VolumetricLightmapData.BrickData.SHCoefficients[5].UAV;
-					PassParameters->OutSkyBentNormal = VolumetricLightmapData.BrickData.SkyBentNormal.UAV;
-					PassParameters->OutDirectionalLightShadowing = VolumetricLightmapData.BrickData.DirectionalLightShadowing.UAV;
-
-					FComputeShaderUtils::AddPass(GraphBuilder, RDG_EVENT_NAME("VolumetricLightmapStitching %d bricks", BricksToCalcThisFrame), ComputeShader, PassParameters, FIntVector(BricksToCalcThisFrame, 1, 1));
-
-					// manually handle transitions since the buffers are not (yet) managed by RDG
-					GraphBuilder.AddPass(RDG_EVENT_NAME("Transition Buffers"), ERDGPassFlags::None,
-						[this](FRHICommandList& RHICmdList) {
-						RHICmdList.Transition({
-							FRHITransitionInfo(AccumulationBrickData.AmbientVector.UAV, ERHIAccess::UAVCompute, ERHIAccess::UAVCompute),
-							FRHITransitionInfo(AccumulationBrickData.SHCoefficients[0].UAV, ERHIAccess::UAVCompute, ERHIAccess::UAVCompute),
-							FRHITransitionInfo(AccumulationBrickData.SHCoefficients[1].UAV, ERHIAccess::UAVCompute, ERHIAccess::UAVCompute),
-							FRHITransitionInfo(AccumulationBrickData.SHCoefficients[2].UAV, ERHIAccess::UAVCompute, ERHIAccess::UAVCompute),
-							FRHITransitionInfo(AccumulationBrickData.SHCoefficients[3].UAV, ERHIAccess::UAVCompute, ERHIAccess::UAVCompute),
-							FRHITransitionInfo(AccumulationBrickData.SHCoefficients[4].UAV, ERHIAccess::UAVCompute, ERHIAccess::UAVCompute),
-							FRHITransitionInfo(AccumulationBrickData.SHCoefficients[5].UAV, ERHIAccess::UAVCompute, ERHIAccess::UAVCompute),
-							FRHITransitionInfo(AccumulationBrickData.SkyBentNormal.UAV, ERHIAccess::UAVCompute, ERHIAccess::UAVCompute),
-							FRHITransitionInfo(AccumulationBrickData.DirectionalLightShadowing.UAV, ERHIAccess::UAVCompute, ERHIAccess::UAVCompute),
-							}
-						);
-					}
-					);
+					PassParameters->OutAmbientVector = OutputBrickDataRDG.AmbientVector;
+					PassParameters->OutSHCoefficients0R = OutputBrickDataRDG.SHCoefficients[0];
+					PassParameters->OutSHCoefficients1R = OutputBrickDataRDG.SHCoefficients[1];
+					PassParameters->OutSHCoefficients0G = OutputBrickDataRDG.SHCoefficients[2];
+					PassParameters->OutSHCoefficients1G = OutputBrickDataRDG.SHCoefficients[3];
+					PassParameters->OutSHCoefficients0B = OutputBrickDataRDG.SHCoefficients[4];
+					PassParameters->OutSHCoefficients1B = OutputBrickDataRDG.SHCoefficients[5];
+					PassParameters->OutSkyBentNormal = OutputBrickDataRDG.SkyBentNormal;
+					PassParameters->OutDirectionalLightShadowing = OutputBrickDataRDG.DirectionalLightShadowing;
+
+					FComputeShaderUtils::AddPass(GraphBuilder, RDG_EVENT_NAME("VolumetricLightmapStitching %d bricks", StitchingBrickBatchSize), ComputeShader, PassParameters,
+												 FIntVector(StitchingBrickBatchSize, 1, 1));
 				}
 			}
+		}
 
-			FrameNumber++;
-
-			SamplesTaken += BricksToCalcThisFrame * NumCellsPerBrick;
+		SamplesTaken += BricksToCalcThisFrame;
 
-			if (SamplesTaken >= (uint64)NumTotalBricks * NumCellsPerBrick * NumTotalPassesToRender)
-			{
-				break;
-			}
+		if (SamplesTaken >= NumTotalSamples)
+		{
+			break;
 		}
 	}
+
+	AccumulationBrickDataRDG.QueueTextureExtractions(GraphBuilder);		
+	OutputBrickDataRDG.QueueTextureExtractions(GraphBuilder);		
+	
 	GraphBuilder.Execute();
+	
+	AccumulationBrickDataRDG.DebugCheckInvariants();
+	OutputBrickDataRDG.DebugCheckInvariants();
 
 	if (IsRayTracingEnabled())
 	{
diff --git a/Engine/Plugins/Experimental/GPULightmass/Source/GPULightmass/Private/VolumetricLightmap.h b/Engine/Plugins/Experimental/GPULightmass/Source/GPULightmass/Private/VolumetricLightmap.h
--- a/Engine/Plugins/Experimental/GPULightmass/Source/GPULightmass/Private/VolumetricLightmap.h
+++ b/Engine/Plugins/Experimental/GPULightmass/Source/GPULightmass/Private/VolumetricLightmap.h
@@ -14,6 +14,8 @@ namespace GPULightmass
 class FVolumetricLightmapRenderer
 {
 public:
+	const int32 BrickBatchSize = 256;
+	
 	FVolumetricLightmapRenderer(FSceneRenderState* Scene);
 
 	void VoxelizeScene();
@@ -44,6 +46,7 @@ private:
 
 	FRWBuffer BrickAllocatorParameters;
 	FRWBuffer BrickRequests;
+	TRefCountPtr<IPooledRenderTarget> ValidityBrickData;
 };
 
 }
diff --git a/Engine/Plugins/Experimental/GPULightmass/Source/GPULightmass/Private/VolumetricLightmapVoxelization.h b/Engine/Plugins/Experimental/GPULightmass/Source/GPULightmass/Private/VolumetricLightmapVoxelization.h
--- a/Engine/Plugins/Experimental/GPULightmass/Source/GPULightmass/Private/VolumetricLightmapVoxelization.h
+++ b/Engine/Plugins/Experimental/GPULightmass/Source/GPULightmass/Private/VolumetricLightmapVoxelization.h
@@ -311,12 +311,60 @@ class FGatherBrickRequestsCS : public FGlobalShader
 	BEGIN_SHADER_PARAMETER_STRUCT(FParameters, )
 		SHADER_PARAMETER(FIntVector, VolumeSize)
 		SHADER_PARAMETER(int32, BrickSize)
+		SHADER_PARAMETER(int32, MipLevel)
 		SHADER_PARAMETER_RDG_TEXTURE_UAV(RWTexture3D<uint>, VoxelizeVolume)
 		SHADER_PARAMETER_UAV(RWBuffer<int>, BrickAllocatorParameters)
+		SHADER_PARAMETER_RDG_BUFFER_UAV(RWBuffer<uint4>, BrickRequestsUnsorted)
+		SHADER_PARAMETER_RDG_BUFFER_UAV(RWBuffer<uint>, BrickRequestKeys)
+		SHADER_PARAMETER_RDG_BUFFER_UAV(RWBuffer<uint>, BrickRequestSortedIndices)
+	END_SHADER_PARAMETER_STRUCT()
+};
+
+class FWriteSortedBrickRequestsCS : public FGlobalShader
+{
+	DECLARE_GLOBAL_SHADER(FWriteSortedBrickRequestsCS);
+	SHADER_USE_PARAMETER_STRUCT(FWriteSortedBrickRequestsCS, FGlobalShader);
+
+	static bool ShouldCompilePermutation(const FGlobalShaderPermutationParameters& Parameters)
+	{
+		return EnumHasAllFlags(Parameters.Flags, EShaderPermutationFlags::HasEditorOnlyData) && RHISupportsRayTracingShaders(Parameters.Platform);
+	}
+
+	static void ModifyCompilationEnvironment(const FGlobalShaderPermutationParameters& Parameters, FShaderCompilerEnvironment& OutEnvironment)
+	{
+		OutEnvironment.CompilerFlags.Add(CFLAG_ForceDXC);
+	}
+
+	BEGIN_SHADER_PARAMETER_STRUCT(FParameters, )
+		SHADER_PARAMETER_RDG_BUFFER_UAV(RWBuffer<uint4>, BrickRequestsUnsorted)
+		SHADER_PARAMETER_RDG_BUFFER_UAV(RWBuffer<uint>, BrickRequestSortedIndices)
+		SHADER_PARAMETER_RDG_BUFFER_UAV(RWBuffer<uint>, BrickRequestSortedIndicesInverse)
 		SHADER_PARAMETER_UAV(RWBuffer<uint4>, BrickRequests)
 	END_SHADER_PARAMETER_STRUCT()
 };
 
+class FPermuteVoxelizeVolumeCS : public FGlobalShader
+{
+	DECLARE_GLOBAL_SHADER(FPermuteVoxelizeVolumeCS);
+	SHADER_USE_PARAMETER_STRUCT(FPermuteVoxelizeVolumeCS, FGlobalShader);
+
+	static bool ShouldCompilePermutation(const FGlobalShaderPermutationParameters& Parameters)
+	{
+		return EnumHasAllFlags(Parameters.Flags, EShaderPermutationFlags::HasEditorOnlyData) && RHISupportsRayTracingShaders(Parameters.Platform);
+	}
+
+	static void ModifyCompilationEnvironment(const FGlobalShaderPermutationParameters& Parameters, FShaderCompilerEnvironment& OutEnvironment)
+	{
+		OutEnvironment.CompilerFlags.Add(CFLAG_ForceDXC);
+	}
+
+	BEGIN_SHADER_PARAMETER_STRUCT(FParameters, )
+		SHADER_PARAMETER(FIntVector, VolumeSize)
+		SHADER_PARAMETER_RDG_BUFFER_UAV(RWBuffer<uint>, BrickRequestSortedIndicesInverse)
+		SHADER_PARAMETER_RDG_TEXTURE_UAV(RWTexture3D<uint>, VoxelizeVolume)
+	END_SHADER_PARAMETER_STRUCT()
+};
+
 class FSplatVolumeCS : public FGlobalShader
 {
 	DECLARE_GLOBAL_SHADER(FSplatVolumeCS);
@@ -364,17 +412,17 @@ class FStitchBorderCS : public FGlobalShader
 		SHADER_PARAMETER(int32, NumTotalBricks)
 		SHADER_PARAMETER(int32, BrickBatchOffset)
 		SHADER_PARAMETER_RDG_TEXTURE_UAV(RWTexture3D<uint4>, IndirectionTexture)
+		SHADER_PARAMETER_RDG_TEXTURE_UAV(RWTexture3D<uint>, ValidityMask)
 		SHADER_PARAMETER_UAV(RWBuffer<uint4>, BrickRequests)
-		SHADER_PARAMETER_TEXTURE(Texture3D<float4>, AmbientVector)
-		SHADER_PARAMETER_UAV(RWTexture3D<float3>, OutAmbientVector)
-		SHADER_PARAMETER_UAV(RWTexture3D<float4>, OutSHCoefficients0R)
-		SHADER_PARAMETER_UAV(RWTexture3D<float4>, OutSHCoefficients1R)
-		SHADER_PARAMETER_UAV(RWTexture3D<float4>, OutSHCoefficients0G)
-		SHADER_PARAMETER_UAV(RWTexture3D<float4>, OutSHCoefficients1G)
-		SHADER_PARAMETER_UAV(RWTexture3D<float4>, OutSHCoefficients0B)
-		SHADER_PARAMETER_UAV(RWTexture3D<float4>, OutSHCoefficients1B)
-		SHADER_PARAMETER_UAV(RWTexture3D<float4>, OutSkyBentNormal)
-		SHADER_PARAMETER_UAV(RWTexture3D<float>, OutDirectionalLightShadowing)
+		SHADER_PARAMETER_RDG_TEXTURE_UAV(RWTexture3D<float3>, OutAmbientVector)
+		SHADER_PARAMETER_RDG_TEXTURE_UAV(RWTexture3D<float4>, OutSHCoefficients0R)
+		SHADER_PARAMETER_RDG_TEXTURE_UAV(RWTexture3D<float4>, OutSHCoefficients1R)
+		SHADER_PARAMETER_RDG_TEXTURE_UAV(RWTexture3D<float4>, OutSHCoefficients0G)
+		SHADER_PARAMETER_RDG_TEXTURE_UAV(RWTexture3D<float4>, OutSHCoefficients1G)
+		SHADER_PARAMETER_RDG_TEXTURE_UAV(RWTexture3D<float4>, OutSHCoefficients0B)
+		SHADER_PARAMETER_RDG_TEXTURE_UAV(RWTexture3D<float4>, OutSHCoefficients1B)
+		SHADER_PARAMETER_RDG_TEXTURE_UAV(RWTexture3D<float4>, OutSkyBentNormal)
+		SHADER_PARAMETER_RDG_TEXTURE_UAV(RWTexture3D<float>, OutDirectionalLightShadowing)
 	END_SHADER_PARAMETER_STRUCT()
 };
 
@@ -396,24 +444,26 @@ class FFinalizeBrickResultsCS : public FGlobalShader
 	BEGIN_SHADER_PARAMETER_STRUCT(FParameters, )
 		SHADER_PARAMETER(int32, NumTotalBricks)
 		SHADER_PARAMETER(int32, BrickBatchOffset)
+		SHADER_PARAMETER(int32, NumTotalPassesToRender)
 		SHADER_PARAMETER_UAV(RWBuffer<uint4>, BrickRequests)
-		SHADER_PARAMETER_TEXTURE(Texture3D<float4>, AmbientVector)
-		SHADER_PARAMETER_TEXTURE(Texture3D<float4>, SHCoefficients0R)
-		SHADER_PARAMETER_TEXTURE(Texture3D<float4>, SHCoefficients1R)
-		SHADER_PARAMETER_TEXTURE(Texture3D<float4>, SHCoefficients0G)
-		SHADER_PARAMETER_TEXTURE(Texture3D<float4>, SHCoefficients1G)
-		SHADER_PARAMETER_TEXTURE(Texture3D<float4>, SHCoefficients0B)
-		SHADER_PARAMETER_TEXTURE(Texture3D<float4>, SHCoefficients1B)
-		SHADER_PARAMETER_TEXTURE(Texture3D<float4>, SkyBentNormal)
-		SHADER_PARAMETER_TEXTURE(Texture3D<float>, DirectionalLightShadowing)
-		SHADER_PARAMETER_UAV(RWTexture3D<float3>, OutAmbientVector)
-		SHADER_PARAMETER_UAV(RWTexture3D<float4>, OutSHCoefficients0R)
-		SHADER_PARAMETER_UAV(RWTexture3D<float4>, OutSHCoefficients1R)
-		SHADER_PARAMETER_UAV(RWTexture3D<float4>, OutSHCoefficients0G)
-		SHADER_PARAMETER_UAV(RWTexture3D<float4>, OutSHCoefficients1G)
-		SHADER_PARAMETER_UAV(RWTexture3D<float4>, OutSHCoefficients0B)
-		SHADER_PARAMETER_UAV(RWTexture3D<float4>, OutSHCoefficients1B)
-		SHADER_PARAMETER_UAV(RWTexture3D<float4>, OutSkyBentNormal)
-		SHADER_PARAMETER_UAV(RWTexture3D<float>, OutDirectionalLightShadowing)
+		SHADER_PARAMETER_RDG_TEXTURE(Texture3D<float4>, AmbientVector)
+		SHADER_PARAMETER_RDG_TEXTURE(Texture3D<float4>, SHCoefficients0R)
+		SHADER_PARAMETER_RDG_TEXTURE(Texture3D<float4>, SHCoefficients1R)
+		SHADER_PARAMETER_RDG_TEXTURE(Texture3D<float4>, SHCoefficients0G)
+		SHADER_PARAMETER_RDG_TEXTURE(Texture3D<float4>, SHCoefficients1G)
+		SHADER_PARAMETER_RDG_TEXTURE(Texture3D<float4>, SHCoefficients0B)
+		SHADER_PARAMETER_RDG_TEXTURE(Texture3D<float4>, SHCoefficients1B)
+		SHADER_PARAMETER_RDG_TEXTURE(Texture3D<float4>, SkyBentNormal)
+		SHADER_PARAMETER_RDG_TEXTURE(Texture3D<float>, DirectionalLightShadowing)
+		SHADER_PARAMETER_RDG_TEXTURE_UAV(RWTexture3D<float3>, OutAmbientVector)
+		SHADER_PARAMETER_RDG_TEXTURE_UAV(RWTexture3D<float4>, OutSHCoefficients0R)
+		SHADER_PARAMETER_RDG_TEXTURE_UAV(RWTexture3D<float4>, OutSHCoefficients1R)
+		SHADER_PARAMETER_RDG_TEXTURE_UAV(RWTexture3D<float4>, OutSHCoefficients0G)
+		SHADER_PARAMETER_RDG_TEXTURE_UAV(RWTexture3D<float4>, OutSHCoefficients1G)
+		SHADER_PARAMETER_RDG_TEXTURE_UAV(RWTexture3D<float4>, OutSHCoefficients0B)
+		SHADER_PARAMETER_RDG_TEXTURE_UAV(RWTexture3D<float4>, OutSHCoefficients1B)
+		SHADER_PARAMETER_RDG_TEXTURE_UAV(RWTexture3D<float4>, OutSkyBentNormal)
+		SHADER_PARAMETER_RDG_TEXTURE_UAV(RWTexture3D<float>, OutDirectionalLightShadowing)
+		SHADER_PARAMETER_RDG_TEXTURE_UAV(RWTexture3D<uint>, ValidityMask)
 	END_SHADER_PARAMETER_STRUCT()
 };
