GPULM: No longer path trace the paddings of each brick. This reduces amount of work from 5x5x5 to 4x4x4 which is 50% smaller. Filtering/stitching passes are separated to accommodate that, which also fixes some leftover VLM seaming issues.

diff --git a/Engine/Plugins/Experimental/GPULightmass/Shaders/Private/BrickAllocationManagement.usf b/Engine/Plugins/Experimental/GPULightmass/Shaders/Private/BrickAllocationManagement.usf
--- a/Engine/Plugins/Experimental/GPULightmass/Shaders/Private/BrickAllocationManagement.usf
+++ b/Engine/Plugins/Experimental/GPULightmass/Shaders/Private/BrickAllocationManagement.usf
@@ -107,7 +107,7 @@ void GatherBrickRequestsCS(
 	if (VoxelizeVolume[VoxelPos])
 	{
 		int LinearIndex = VoxelizeVolume[VoxelPos] - 1;
-		BrickRequestKeys[LinearIndex] = (MipLevel << 30) | MortonCode3(VoxelPos.x) | (MortonCode3(VoxelPos.y) << 1) | (MortonCode3(VoxelPos.z) << 2);
+		BrickRequestKeys[LinearIndex] = ((2 - MipLevel) << 30) | MortonCode3(VoxelPos.x) | (MortonCode3(VoxelPos.y) << 1) | (MortonCode3(VoxelPos.z) << 2);
 		BrickRequestSortedIndices[LinearIndex] = LinearIndex;
 		BrickRequestsUnsorted[LinearIndex] = uint4(VoxelPos, BrickSize);
 	}
@@ -381,7 +381,41 @@ int BrickBatchOffset;
 RWTexture3D<uint> ValidityMask;
 
 [numthreads(5, 5, 5)]
-void StitchBorderCS(uint3 BrickID : SV_GroupID, uint3 CellPosInBrick : SV_GroupThreadID)
+void CopyPaddingStripsCS(uint3 BrickID : SV_GroupID, uint3 CellPosInBrick : SV_GroupThreadID)
+{
+	uint BrickIndex = BrickID.x + BrickBatchOffset;
+	if (BrickIndex >= NumTotalBricks) return;	
+	int3 CellPosInCurrentMip = BrickRequests[BrickIndex].xyz * 4 + CellPosInBrick;	
+	float3 InvCellSizeInCurrentMip = float3(1, 1, 1) / 4.0f;
+	float PaddedBrickSize = 4 + 1;
+	uint3 VoxelPos = ComputeBrickLayoutPosition(BrickIndex, uint3(256, 256, 256)) * PaddedBrickSize + CellPosInBrick;
+
+	if (BrickRequests[BrickIndex].w == 1 << (2 * MipLevel))
+	{
+		int3 VoxelizeVolumeCoords = floor(CellPosInCurrentMip * InvCellSizeInCurrentMip);
+		int BrickLinearAddress = VoxelizeVolume[VoxelizeVolumeCoords] - 1;
+
+		if (BrickLinearAddress >= 0)
+		{
+			int3 BrickOrigin = ComputeBrickLayoutPosition(BrickLinearAddress, uint3(256, 256, 256));
+			int3 BrickUVW = 5 * BrickOrigin + 4 * frac(CellPosInCurrentMip * InvCellSizeInCurrentMip);
+		
+			OutAmbientVector[VoxelPos] = OutAmbientVector[BrickUVW];
+			OutSHCoefficients0R[VoxelPos] = OutSHCoefficients0R[BrickUVW];
+			OutSHCoefficients1R[VoxelPos] = OutSHCoefficients1R[BrickUVW];
+			OutSHCoefficients0G[VoxelPos] = OutSHCoefficients0G[BrickUVW];
+			OutSHCoefficients1G[VoxelPos] = OutSHCoefficients1G[BrickUVW];
+			OutSHCoefficients0B[VoxelPos] = OutSHCoefficients0B[BrickUVW];
+			OutSHCoefficients1B[VoxelPos] = OutSHCoefficients1B[BrickUVW];
+			OutSkyBentNormal[VoxelPos] = OutSkyBentNormal[BrickUVW];
+			OutDirectionalLightShadowing[VoxelPos] = OutDirectionalLightShadowing[BrickUVW];
+			ValidityMask[VoxelPos] = ValidityMask[BrickUVW];
+		}
+	}
+}
+
+[numthreads(5, 5, 5)]
+void FillInvalidCellCS(uint3 BrickID : SV_GroupID, uint3 CellPosInBrick : SV_GroupThreadID)
 {
 	uint BrickIndex = BrickID.x + BrickBatchOffset;
 	if (BrickIndex >= NumTotalBricks) return;	
@@ -403,6 +437,17 @@ void StitchBorderCS(uint3 BrickID : SV_GroupID, uint3 CellPosInBrick : SV_GroupT
 		OutSkyBentNormal[VoxelPos] = DilateWithValidityMask(OutSkyBentNormal, ValidityMask, IndirectionTextureUVW);
 		OutDirectionalLightShadowing[VoxelPos] = DilateWithValidityMask(OutDirectionalLightShadowing, ValidityMask, IndirectionTextureUVW);
 	}
+}
+
+[numthreads(5, 5, 5)]
+void StitchBorderCS(uint3 BrickID : SV_GroupID, uint3 CellPosInBrick : SV_GroupThreadID)
+{
+	uint BrickIndex = BrickID.x + BrickBatchOffset;
+	if (BrickIndex >= NumTotalBricks) return;	
+	int3 CellPosInVLM = (BrickRequests[BrickIndex].xyz * 4 + CellPosInBrick) * BrickRequests[BrickIndex].w;	
+	float3 InvCellSizeInVLM = float3(1, 1, 1) / 4.0f;
+	float PaddedBrickSize = 4 + 1;
+	uint3 VoxelPos = ComputeBrickLayoutPosition(BrickIndex, uint3(256, 256, 256)) * PaddedBrickSize + CellPosInBrick;
 
 	for (int dx = -1; dx <= 1; dx++)
 	{
diff --git a/Engine/Plugins/Experimental/GPULightmass/Shaders/Private/LightmapPathTracing.usf b/Engine/Plugins/Experimental/GPULightmass/Shaders/Private/LightmapPathTracing.usf
--- a/Engine/Plugins/Experimental/GPULightmass/Shaders/Private/LightmapPathTracing.usf
+++ b/Engine/Plugins/Experimental/GPULightmass/Shaders/Private/LightmapPathTracing.usf
@@ -1214,11 +1214,11 @@ StructuredBuffer<FLightShaderConstants> LightShaderParametersArray;
 [shader("raygeneration")]
 void VolumetricLightmapPathTracingMainRG()
 {
-	uint BrickVolumeSize = 5 * 5 * 5;
+	uint BrickVolumeSize = 4 * 4 * 4;
 	uint BrickIndex = DispatchRaysIndex().x / BrickVolumeSize + BrickBatchOffset;
 	if (BrickIndex >= NumTotalBricks) return;
 	uint CellIndex = DispatchRaysIndex().x % BrickVolumeSize;
-	uint3 CellPosInBrick = ComputeBrickLayoutPosition(CellIndex, uint3(5, 5, 5));
+	uint3 CellPosInBrick = ComputeBrickLayoutPosition(CellIndex, uint3(4, 4, 4));
 	uint PaddedBrickSize = 4 + 1;
 	uint3 VoxelPos = ComputeBrickLayoutPosition(DispatchRaysIndex().x / BrickVolumeSize, uint3(256, 256, 256)) * PaddedBrickSize + CellPosInBrick;
 	
diff --git a/Engine/Plugins/Experimental/GPULightmass/Source/GPULightmass/Private/VolumetricLightmap.cpp b/Engine/Plugins/Experimental/GPULightmass/Source/GPULightmass/Private/VolumetricLightmap.cpp
--- a/Engine/Plugins/Experimental/GPULightmass/Source/GPULightmass/Private/VolumetricLightmap.cpp
+++ b/Engine/Plugins/Experimental/GPULightmass/Source/GPULightmass/Private/VolumetricLightmap.cpp
@@ -24,6 +24,8 @@ IMPLEMENT_GLOBAL_SHADER(FCountNumBricksCS, "/Plugin/GPULightmass/Private/BrickAl
 IMPLEMENT_GLOBAL_SHADER(FGatherBrickRequestsCS, "/Plugin/GPULightmass/Private/BrickAllocationManagement.usf", "GatherBrickRequestsCS", SF_Compute);
 IMPLEMENT_GLOBAL_SHADER(FSplatVolumeCS, "/Plugin/GPULightmass/Private/BrickAllocationManagement.usf", "SplatVolumeCS", SF_Compute);
 IMPLEMENT_GLOBAL_SHADER(FStitchBorderCS, "/Plugin/GPULightmass/Private/BrickAllocationManagement.usf", "StitchBorderCS", SF_Compute);
+IMPLEMENT_GLOBAL_SHADER(FCopyPaddingStripsCS, "/Plugin/GPULightmass/Private/BrickAllocationManagement.usf", "CopyPaddingStripsCS", SF_Compute);
+IMPLEMENT_GLOBAL_SHADER(FFillInvalidCellCS, "/Plugin/GPULightmass/Private/BrickAllocationManagement.usf", "FillInvalidCellCS", SF_Compute);
 IMPLEMENT_GLOBAL_SHADER(FFinalizeBrickResultsCS, "/Plugin/GPULightmass/Private/BrickAllocationManagement.usf", "FinalizeBrickResultsCS", SF_Compute);
 IMPLEMENT_GLOBAL_SHADER(FWriteSortedBrickRequestsCS, "/Plugin/GPULightmass/Private/BrickAllocationManagement.usf", "WriteSortedBrickRequestsCS", SF_Compute);
 IMPLEMENT_GLOBAL_SHADER(FPermuteVoxelizeVolumeCS, "/Plugin/GPULightmass/Private/BrickAllocationManagement.usf", "PermuteVoxelizeVolumeCS", SF_Compute);
@@ -456,7 +458,20 @@ void FVolumetricLightmapRenderer::VoxelizeScene()
 	}
 
 	InitializeBrickData(BrickLayoutDimensions * 5, VolumetricLightmapData.BrickData, false);
-	InitializeBrickData(FIntVector(256, 1, 1) * 5, AccumulationBrickData, true);
+
+	FIntVector BrickLayoutDimensionsForAccumulation;
+
+	{
+		int32 BrickTextureLinearAllocator = BrickBatchSize;
+		BrickLayoutDimensionsForAccumulation.X = FMath::Min(BrickTextureLinearAllocator, MaxBricksInLayoutOneDim);
+		BrickTextureLinearAllocator = FMath::DivideAndRoundUp(BrickTextureLinearAllocator, BrickLayoutDimensionsForAccumulation.X);
+		BrickLayoutDimensionsForAccumulation.Y = FMath::Min(BrickTextureLinearAllocator, MaxBricksInLayoutOneDim);
+		BrickTextureLinearAllocator = FMath::DivideAndRoundUp(BrickTextureLinearAllocator, BrickLayoutDimensionsForAccumulation.Y);
+		BrickLayoutDimensionsForAccumulation.Z = FMath::Min(BrickTextureLinearAllocator, MaxBricksInLayoutOneDim);
+	}
+	
+	InitializeBrickData(BrickLayoutDimensionsForAccumulation * 5, AccumulationBrickData, true);
+	
 	BrickRequests.Initialize(TEXT("BrickRequests"), 16, NumTotalBricks, PF_R32G32B32A32_UINT, BUF_UnorderedAccess);
 
 	VolumetricLightmapData.BrickDataDimensions = BrickLayoutDimensions * 5;
@@ -745,7 +760,7 @@ void FVolumetricLightmapRenderer::BackgroundTick()
 	FGlobalShaderMap* GlobalShaderMap = GetGlobalShaderMap(Scene->FeatureLevel);
 
 	const bool bIsViewportNonRealtime = GCurrentLevelEditingViewportClient && !GCurrentLevelEditingViewportClient->IsRealtime();
-	const int32 NumSamplesPerFrame = 8 * (bIsViewportNonRealtime ? Scene->Settings->TilePassesInFullSpeedMode : Scene->Settings->TilePassesInSlowMode);
+	const int32 NumSamplesPerFrame = bIsViewportNonRealtime ? Scene->Settings->TilePassesInFullSpeedMode : Scene->Settings->TilePassesInSlowMode;
 
 	FVolumetricLightmapPathTracingRGS::FParameters* PreviousPassParameters = nullptr;
 	
@@ -834,7 +849,7 @@ void FVolumetricLightmapRenderer::BackgroundTick()
 
 			FSceneRenderState* SceneRenderState = Scene; // capture member variable
 			GraphBuilder.AddPass(
-				RDG_EVENT_NAME("VolumetricLightmapPathTracing %d bricks %d rays", BricksToCalcThisFrame, BricksToCalcThisFrame * (BrickSize + 1) * (BrickSize + 1) * (BrickSize + 1)),
+				RDG_EVENT_NAME("VolumetricLightmapPathTracing %d bricks %d rays", BricksToCalcThisFrame, BricksToCalcThisFrame * BrickSize * BrickSize * BrickSize),
 				PassParameters,
 				ERDGPassFlags::Compute,
 				[
@@ -849,7 +864,7 @@ void FVolumetricLightmapRenderer::BackgroundTick()
 					SetShaderParameters(GlobalResources, RayGenShader, *PassParameters);
 
 					RHICmdList.RayTraceDispatch(RayTracingPipelineState, RayGenShader.GetRayTracingShader(), RayTracingSceneRHI, GlobalResources,
-					                            BricksToCalcThisFrame * (BrickSize + 1) * (BrickSize + 1) * (BrickSize + 1), 1);
+					                            BricksToCalcThisFrame * BrickSize * BrickSize * BrickSize, 1);
 				}
 			);
 		}
@@ -890,16 +905,138 @@ void FVolumetricLightmapRenderer::BackgroundTick()
 			FComputeShaderUtils::AddPass(GraphBuilder, RDG_EVENT_NAME("FinalizeBrickResults"), ComputeShader, PassParameters, FIntVector(BricksToCalcThisFrame, 1, 1));
 		}
 
+		FRDGTextureUAV* IndirectionTextureUAV = GraphBuilder.CreateUAV(GraphBuilder.RegisterExternalTexture(IndirectionTexture));
+		
+		// Do some temporary copy and fill passes for preview
+		if (!bIsViewportNonRealtime)
+		{
+			for (int32 MipLevel = VoxelizationVolumeMips.Num() - 1; MipLevel >= 0; MipLevel--)
+			{
+				// Copy padding strips from the neighbors of each brick
+				TShaderMapRef<FCopyPaddingStripsCS> ComputeShader(GlobalShaderMap);
+
+				FCopyPaddingStripsCS::FParameters* PassParameters = GraphBuilder.AllocParameters<FCopyPaddingStripsCS::FParameters>();
+				PassParameters->BrickDataDimensions = VolumetricLightmapData.BrickDataDimensions;
+				PassParameters->IndirectionTextureDim = IndirectionTextureDimensions;
+				PassParameters->FrameNumber = NumTotalPassesToRender;
+				PassParameters->NumTotalBricks = NumTotalBricks;
+				PassParameters->BrickBatchOffset = BrickBatchOffset;
+				PassParameters->IndirectionTexture = IndirectionTextureUAV;
+				PassParameters->MipLevel = MipLevel;
+				PassParameters->VoxelizeVolume = GraphBuilder.CreateUAV(GraphBuilder.RegisterExternalTexture(VoxelizationVolumeMips[MipLevel]));
+				PassParameters->ValidityMask = ValidityBrickDataUAV;
+				PassParameters->BrickRequests = BrickRequests.UAV;
+				PassParameters->OutAmbientVector = OutputBrickDataRDG.AmbientVector;
+				PassParameters->OutSHCoefficients0R = OutputBrickDataRDG.SHCoefficients[0];
+				PassParameters->OutSHCoefficients1R = OutputBrickDataRDG.SHCoefficients[1];
+				PassParameters->OutSHCoefficients0G = OutputBrickDataRDG.SHCoefficients[2];
+				PassParameters->OutSHCoefficients1G = OutputBrickDataRDG.SHCoefficients[3];
+				PassParameters->OutSHCoefficients0B = OutputBrickDataRDG.SHCoefficients[4];
+				PassParameters->OutSHCoefficients1B = OutputBrickDataRDG.SHCoefficients[5];
+				PassParameters->OutSkyBentNormal = OutputBrickDataRDG.SkyBentNormal;
+				PassParameters->OutDirectionalLightShadowing = OutputBrickDataRDG.DirectionalLightShadowing;
+
+				FComputeShaderUtils::AddPass(GraphBuilder, RDG_EVENT_NAME("VolumetricLightmapCopyPaddingStrips %d bricks", BrickBatchSize), ComputeShader, PassParameters,
+				                             FIntVector(BrickBatchSize, 1, 1));
+			}
+
+			if (RenderPassIndex == NumTotalPassesToRender - 1)
+			{
+				TShaderMapRef<FFillInvalidCellCS> ComputeShader(GlobalShaderMap);
+
+				FFillInvalidCellCS::FParameters* PassParameters = GraphBuilder.AllocParameters<FFillInvalidCellCS::FParameters>();
+				PassParameters->BrickDataDimensions = VolumetricLightmapData.BrickDataDimensions;
+				PassParameters->IndirectionTextureDim = IndirectionTextureDimensions;
+				PassParameters->FrameNumber = NumTotalPassesToRender;
+				PassParameters->NumTotalBricks = NumTotalBricks;
+				PassParameters->BrickBatchOffset = BrickBatchOffset;
+				PassParameters->IndirectionTexture = IndirectionTextureUAV;
+				PassParameters->ValidityMask = ValidityBrickDataUAV;
+				PassParameters->BrickRequests = BrickRequests.UAV;
+				PassParameters->OutAmbientVector = OutputBrickDataRDG.AmbientVector;
+				PassParameters->OutSHCoefficients0R = OutputBrickDataRDG.SHCoefficients[0];
+				PassParameters->OutSHCoefficients1R = OutputBrickDataRDG.SHCoefficients[1];
+				PassParameters->OutSHCoefficients0G = OutputBrickDataRDG.SHCoefficients[2];
+				PassParameters->OutSHCoefficients1G = OutputBrickDataRDG.SHCoefficients[3];
+				PassParameters->OutSHCoefficients0B = OutputBrickDataRDG.SHCoefficients[4];
+				PassParameters->OutSHCoefficients1B = OutputBrickDataRDG.SHCoefficients[5];
+				PassParameters->OutSkyBentNormal = OutputBrickDataRDG.SkyBentNormal;
+				PassParameters->OutDirectionalLightShadowing = OutputBrickDataRDG.DirectionalLightShadowing;
+
+				FComputeShaderUtils::AddPass(GraphBuilder, RDG_EVENT_NAME("VolumetricLightmapFillInvalidCell %d bricks", BrickBatchSize), ComputeShader, PassParameters,
+											 FIntVector(BrickBatchSize, 1, 1));
+			}
+		}
+
 		// Do stitching over all brickes if this is the very last pass of all bricks
 		if (SamplesTaken < NumTotalSamples && SamplesTaken + BricksToCalcThisFrame >= NumTotalSamples)
 		{
-			FRDGTextureUAV* IndirectionTextureUAV = GraphBuilder.CreateUAV(GraphBuilder.RegisterExternalTexture(IndirectionTexture));
+			const int32 StitchingBrickBatchSize = 256;
+
+			for (int32 MipLevel = VoxelizationVolumeMips.Num() - 1; MipLevel >= 0; MipLevel--)
+			{
+				// Copy padding strips from the neighbors of each brick
+				for (int32 StitchingBrickBatchOffset = 0; StitchingBrickBatchOffset < NumTotalBricks; StitchingBrickBatchOffset += StitchingBrickBatchSize)
+				{
+					TShaderMapRef<FCopyPaddingStripsCS> ComputeShader(GlobalShaderMap);
+
+					FCopyPaddingStripsCS::FParameters* PassParameters = GraphBuilder.AllocParameters<FCopyPaddingStripsCS::FParameters>();
+					PassParameters->BrickDataDimensions = VolumetricLightmapData.BrickDataDimensions;
+					PassParameters->IndirectionTextureDim = IndirectionTextureDimensions;
+					PassParameters->FrameNumber = NumTotalPassesToRender;
+					PassParameters->NumTotalBricks = NumTotalBricks;
+					PassParameters->BrickBatchOffset = StitchingBrickBatchOffset;
+					PassParameters->IndirectionTexture = IndirectionTextureUAV;
+					PassParameters->MipLevel = MipLevel;
+					PassParameters->VoxelizeVolume = GraphBuilder.CreateUAV(GraphBuilder.RegisterExternalTexture(VoxelizationVolumeMips[MipLevel]));
+					PassParameters->ValidityMask = ValidityBrickDataUAV;
+					PassParameters->BrickRequests = BrickRequests.UAV;
+					PassParameters->OutAmbientVector = OutputBrickDataRDG.AmbientVector;
+					PassParameters->OutSHCoefficients0R = OutputBrickDataRDG.SHCoefficients[0];
+					PassParameters->OutSHCoefficients1R = OutputBrickDataRDG.SHCoefficients[1];
+					PassParameters->OutSHCoefficients0G = OutputBrickDataRDG.SHCoefficients[2];
+					PassParameters->OutSHCoefficients1G = OutputBrickDataRDG.SHCoefficients[3];
+					PassParameters->OutSHCoefficients0B = OutputBrickDataRDG.SHCoefficients[4];
+					PassParameters->OutSHCoefficients1B = OutputBrickDataRDG.SHCoefficients[5];
+					PassParameters->OutSkyBentNormal = OutputBrickDataRDG.SkyBentNormal;
+					PassParameters->OutDirectionalLightShadowing = OutputBrickDataRDG.DirectionalLightShadowing;
 
+					FComputeShaderUtils::AddPass(GraphBuilder, RDG_EVENT_NAME("VolumetricLightmapCopyPaddingStrips %d bricks", StitchingBrickBatchSize), ComputeShader, PassParameters,
+												 FIntVector(StitchingBrickBatchSize, 1, 1));
+				}
+			}
+
+			// Fill invalid cells by dilation
+			for (int32 StitchingBrickBatchOffset = 0; StitchingBrickBatchOffset < NumTotalBricks; StitchingBrickBatchOffset += StitchingBrickBatchSize)
+			{
+				TShaderMapRef<FFillInvalidCellCS> ComputeShader(GlobalShaderMap);
+
+				FFillInvalidCellCS::FParameters* PassParameters = GraphBuilder.AllocParameters<FFillInvalidCellCS::FParameters>();
+				PassParameters->BrickDataDimensions = VolumetricLightmapData.BrickDataDimensions;
+				PassParameters->IndirectionTextureDim = IndirectionTextureDimensions;
+				PassParameters->FrameNumber = NumTotalPassesToRender;
+				PassParameters->NumTotalBricks = NumTotalBricks;
+				PassParameters->BrickBatchOffset = StitchingBrickBatchOffset;
+				PassParameters->IndirectionTexture = IndirectionTextureUAV;
+				PassParameters->ValidityMask = ValidityBrickDataUAV;
+				PassParameters->BrickRequests = BrickRequests.UAV;
+				PassParameters->OutAmbientVector = OutputBrickDataRDG.AmbientVector;
+				PassParameters->OutSHCoefficients0R = OutputBrickDataRDG.SHCoefficients[0];
+				PassParameters->OutSHCoefficients1R = OutputBrickDataRDG.SHCoefficients[1];
+				PassParameters->OutSHCoefficients0G = OutputBrickDataRDG.SHCoefficients[2];
+				PassParameters->OutSHCoefficients1G = OutputBrickDataRDG.SHCoefficients[3];
+				PassParameters->OutSHCoefficients0B = OutputBrickDataRDG.SHCoefficients[4];
+				PassParameters->OutSHCoefficients1B = OutputBrickDataRDG.SHCoefficients[5];
+				PassParameters->OutSkyBentNormal = OutputBrickDataRDG.SkyBentNormal;
+				PassParameters->OutDirectionalLightShadowing = OutputBrickDataRDG.DirectionalLightShadowing;
+
+				FComputeShaderUtils::AddPass(GraphBuilder, RDG_EVENT_NAME("VolumetricLightmapFillInvalidCell %d bricks", StitchingBrickBatchSize), ComputeShader, PassParameters,
+											 FIntVector(StitchingBrickBatchSize, 1, 1));
+			}
+			
 			// Do 2 passes to propagate across 3 refinement levels
 			for (int32 StitchPass = 0; StitchPass < 2; StitchPass++)
 			{
-				const int32 StitchingBrickBatchSize = 256;
-
 				for (int32 StitchingBrickBatchOffset = 0; StitchingBrickBatchOffset < NumTotalBricks; StitchingBrickBatchOffset += StitchingBrickBatchSize)
 				{
 					TShaderMapRef<FStitchBorderCS> ComputeShader(GlobalShaderMap);
diff --git a/Engine/Plugins/Experimental/GPULightmass/Source/GPULightmass/Private/VolumetricLightmap.h b/Engine/Plugins/Experimental/GPULightmass/Source/GPULightmass/Private/VolumetricLightmap.h
--- a/Engine/Plugins/Experimental/GPULightmass/Source/GPULightmass/Private/VolumetricLightmap.h
+++ b/Engine/Plugins/Experimental/GPULightmass/Source/GPULightmass/Private/VolumetricLightmap.h
@@ -14,7 +14,7 @@ namespace GPULightmass
 class FVolumetricLightmapRenderer
 {
 public:
-	const int32 BrickBatchSize = 256;
+	const int32 BrickBatchSize = 2048;
 	
 	FVolumetricLightmapRenderer(FSceneRenderState* Scene);
 
diff --git a/Engine/Plugins/Experimental/GPULightmass/Source/GPULightmass/Private/VolumetricLightmapVoxelization.h b/Engine/Plugins/Experimental/GPULightmass/Source/GPULightmass/Private/VolumetricLightmapVoxelization.h
--- a/Engine/Plugins/Experimental/GPULightmass/Source/GPULightmass/Private/VolumetricLightmapVoxelization.h
+++ b/Engine/Plugins/Experimental/GPULightmass/Source/GPULightmass/Private/VolumetricLightmapVoxelization.h
@@ -426,6 +426,80 @@ class FStitchBorderCS : public FGlobalShader
 	END_SHADER_PARAMETER_STRUCT()
 };
 
+class FCopyPaddingStripsCS : public FGlobalShader
+{
+	DECLARE_GLOBAL_SHADER(FCopyPaddingStripsCS);
+	SHADER_USE_PARAMETER_STRUCT(FCopyPaddingStripsCS, FGlobalShader);
+
+	static bool ShouldCompilePermutation(const FGlobalShaderPermutationParameters& Parameters)
+	{
+		return EnumHasAllFlags(Parameters.Flags, EShaderPermutationFlags::HasEditorOnlyData) && RHISupportsRayTracingShaders(Parameters.Platform);
+	}
+
+	static void ModifyCompilationEnvironment(const FGlobalShaderPermutationParameters& Parameters, FShaderCompilerEnvironment& OutEnvironment)
+	{
+		OutEnvironment.CompilerFlags.Add(CFLAG_ForceDXC);
+	}
+
+	BEGIN_SHADER_PARAMETER_STRUCT(FParameters, )
+		SHADER_PARAMETER(FIntVector, IndirectionTextureDim)
+		SHADER_PARAMETER(FIntVector, BrickDataDimensions)
+		SHADER_PARAMETER(uint32, FrameNumber)
+		SHADER_PARAMETER(int32, MipLevel)
+		SHADER_PARAMETER(int32, NumTotalBricks)
+		SHADER_PARAMETER(int32, BrickBatchOffset)
+		SHADER_PARAMETER_RDG_TEXTURE_UAV(RWTexture3D<uint4>, IndirectionTexture)	
+		SHADER_PARAMETER_RDG_TEXTURE_UAV(RWTexture3D<uint>, VoxelizeVolume)
+		SHADER_PARAMETER_RDG_TEXTURE_UAV(RWTexture3D<uint>, ValidityMask)
+		SHADER_PARAMETER_UAV(RWBuffer<uint4>, BrickRequests)
+		SHADER_PARAMETER_RDG_TEXTURE_UAV(RWTexture3D<float3>, OutAmbientVector)
+		SHADER_PARAMETER_RDG_TEXTURE_UAV(RWTexture3D<float4>, OutSHCoefficients0R)
+		SHADER_PARAMETER_RDG_TEXTURE_UAV(RWTexture3D<float4>, OutSHCoefficients1R)
+		SHADER_PARAMETER_RDG_TEXTURE_UAV(RWTexture3D<float4>, OutSHCoefficients0G)
+		SHADER_PARAMETER_RDG_TEXTURE_UAV(RWTexture3D<float4>, OutSHCoefficients1G)
+		SHADER_PARAMETER_RDG_TEXTURE_UAV(RWTexture3D<float4>, OutSHCoefficients0B)
+		SHADER_PARAMETER_RDG_TEXTURE_UAV(RWTexture3D<float4>, OutSHCoefficients1B)
+		SHADER_PARAMETER_RDG_TEXTURE_UAV(RWTexture3D<float4>, OutSkyBentNormal)
+		SHADER_PARAMETER_RDG_TEXTURE_UAV(RWTexture3D<float>, OutDirectionalLightShadowing)
+	END_SHADER_PARAMETER_STRUCT()
+};
+
+class FFillInvalidCellCS : public FGlobalShader
+{
+	DECLARE_GLOBAL_SHADER(FFillInvalidCellCS);
+	SHADER_USE_PARAMETER_STRUCT(FFillInvalidCellCS, FGlobalShader);
+
+	static bool ShouldCompilePermutation(const FGlobalShaderPermutationParameters& Parameters)
+	{
+		return EnumHasAllFlags(Parameters.Flags, EShaderPermutationFlags::HasEditorOnlyData) && RHISupportsRayTracingShaders(Parameters.Platform);
+	}
+
+	static void ModifyCompilationEnvironment(const FGlobalShaderPermutationParameters& Parameters, FShaderCompilerEnvironment& OutEnvironment)
+	{
+		OutEnvironment.CompilerFlags.Add(CFLAG_ForceDXC);
+	}
+
+	BEGIN_SHADER_PARAMETER_STRUCT(FParameters, )
+		SHADER_PARAMETER(FIntVector, IndirectionTextureDim)
+		SHADER_PARAMETER(FIntVector, BrickDataDimensions)
+		SHADER_PARAMETER(uint32, FrameNumber)
+		SHADER_PARAMETER(int32, NumTotalBricks)
+		SHADER_PARAMETER(int32, BrickBatchOffset)
+		SHADER_PARAMETER_RDG_TEXTURE_UAV(RWTexture3D<uint4>, IndirectionTexture)
+		SHADER_PARAMETER_RDG_TEXTURE_UAV(RWTexture3D<uint>, ValidityMask)
+		SHADER_PARAMETER_UAV(RWBuffer<uint4>, BrickRequests)
+		SHADER_PARAMETER_RDG_TEXTURE_UAV(RWTexture3D<float3>, OutAmbientVector)
+		SHADER_PARAMETER_RDG_TEXTURE_UAV(RWTexture3D<float4>, OutSHCoefficients0R)
+		SHADER_PARAMETER_RDG_TEXTURE_UAV(RWTexture3D<float4>, OutSHCoefficients1R)
+		SHADER_PARAMETER_RDG_TEXTURE_UAV(RWTexture3D<float4>, OutSHCoefficients0G)
+		SHADER_PARAMETER_RDG_TEXTURE_UAV(RWTexture3D<float4>, OutSHCoefficients1G)
+		SHADER_PARAMETER_RDG_TEXTURE_UAV(RWTexture3D<float4>, OutSHCoefficients0B)
+		SHADER_PARAMETER_RDG_TEXTURE_UAV(RWTexture3D<float4>, OutSHCoefficients1B)
+		SHADER_PARAMETER_RDG_TEXTURE_UAV(RWTexture3D<float4>, OutSkyBentNormal)
+		SHADER_PARAMETER_RDG_TEXTURE_UAV(RWTexture3D<float>, OutDirectionalLightShadowing)
+	END_SHADER_PARAMETER_STRUCT()
+};
+
 class FFinalizeBrickResultsCS : public FGlobalShader
 {
 	DECLARE_GLOBAL_SHADER(FFinalizeBrickResultsCS);
