GPULM: Tighten volumetric lightmap voxelization boundaries to reduce number of bricks created. Now we conservatively voxelize meshes at 4x resolution to detect if any cells are going to be placed on stitching boundaries and generate extra cells only when necessary, instead of brute force dilation.

diff --git a/Engine/Plugins/Experimental/GPULightmass/Shaders/Private/BrickAllocationManagement.usf b/Engine/Plugins/Experimental/GPULightmass/Shaders/Private/BrickAllocationManagement.usf
--- a/Engine/Plugins/Experimental/GPULightmass/Shaders/Private/BrickAllocationManagement.usf
+++ b/Engine/Plugins/Experimental/GPULightmass/Shaders/Private/BrickAllocationManagement.usf
@@ -51,25 +51,6 @@ void DownsampleVolumeCS(
 	}
 }
 
-[numthreads(4, 4, 4)]
-void DilateVolumeCS(uint3 VoxelPos : SV_DispatchThreadID) 
-{
-	if (any(VoxelPos >= VolumeSize)) return; 
-	for (int dx = -1; dx <= 1; dx++)
-	{
-		for (int dy = -1; dy <= 1; dy++)
-		{
-			for (int dz = -1; dz <= 1; dz++)
-			{
-				if (VoxelizeVolume[VoxelPos] == BRICK_IN_IMPORTANCE_VOLUME && VoxelizeVolume[int3(VoxelPos) + int3(dx, dy, dz)] == BRICK_ALLOCATED)
-				{
-					VoxelizeVolume[VoxelPos] = BRICK_DILATED;
-				}
-			}
-		}
-	}
-}
-
 RWTexture3D<uint4> IndirectionTexture;
 
 RWBuffer<uint4> BrickRequests;
diff --git a/Engine/Plugins/Experimental/GPULightmass/Shaders/Private/VolumetricLightmapVoxelization.usf b/Engine/Plugins/Experimental/GPULightmass/Shaders/Private/VolumetricLightmapVoxelization.usf
--- a/Engine/Plugins/Experimental/GPULightmass/Shaders/Private/VolumetricLightmapVoxelization.usf
+++ b/Engine/Plugins/Experimental/GPULightmass/Shaders/Private/VolumetricLightmapVoxelization.usf
@@ -12,6 +12,7 @@
 #include "BrickAllocationDefs.ush"
 
 #define SOFTWARE_CONSERVATIVE_RASTERIZATION 1
+#define SUPER_RESOLUTION_FACTOR 4 // Must match VolumetricLightmap.cpp
 
 struct FVLMVoxelizationVSToGS
 {
@@ -100,8 +101,8 @@ void VLMVoxelizationGS(triangle FVLMVoxelizationVSToGS Inputs[3], inout Triangle
 	AABB.zw = max(AABB.zw, Vertices[1].ScreenPosition.xy);
 	AABB.xy = min(AABB.xy, Vertices[2].ScreenPosition.xy);
 	AABB.zw = max(AABB.zw, Vertices[2].ScreenPosition.xy);
-	AABB.xy -= float2(1.0f / VLMVoxelizationParams.VolumeMaxDim, 1.0f / VLMVoxelizationParams.VolumeMaxDim);
-	AABB.zw += float2(1.0f / VLMVoxelizationParams.VolumeMaxDim, 1.0f / VLMVoxelizationParams.VolumeMaxDim);
+	AABB.xy -= float2(1.0f / VLMVoxelizationParams.VolumeMaxDim / SUPER_RESOLUTION_FACTOR, 1.0f / VLMVoxelizationParams.VolumeMaxDim / SUPER_RESOLUTION_FACTOR);
+	AABB.zw += float2(1.0f / VLMVoxelizationParams.VolumeMaxDim / SUPER_RESOLUTION_FACTOR, 1.0f / VLMVoxelizationParams.VolumeMaxDim / SUPER_RESOLUTION_FACTOR);
 	
 	Vertices[0].AABB = AABB;
 	Vertices[1].AABB = AABB;
@@ -114,10 +115,10 @@ void VLMVoxelizationGS(triangle FVLMVoxelizationVSToGS Inputs[3], inout Triangle
 	};
 	
 	float3 SemiDiagonalVectors[4] = {
-		{ sqrt(2.0f) / VLMVoxelizationParams.VolumeMaxDim,  sqrt(2.0f) / VLMVoxelizationParams.VolumeMaxDim, 0.0f},
-		{ sqrt(2.0f) / VLMVoxelizationParams.VolumeMaxDim, -sqrt(2.0f) / VLMVoxelizationParams.VolumeMaxDim, 0.0f},
-		{-sqrt(2.0f) / VLMVoxelizationParams.VolumeMaxDim,  sqrt(2.0f) / VLMVoxelizationParams.VolumeMaxDim, 0.0f},
-		{-sqrt(2.0f) / VLMVoxelizationParams.VolumeMaxDim, -sqrt(2.0f) / VLMVoxelizationParams.VolumeMaxDim, 0.0f}
+		{ sqrt(2.0f) / VLMVoxelizationParams.VolumeMaxDim / SUPER_RESOLUTION_FACTOR,  sqrt(2.0f) / VLMVoxelizationParams.VolumeMaxDim / SUPER_RESOLUTION_FACTOR, 0.0f},
+		{ sqrt(2.0f) / VLMVoxelizationParams.VolumeMaxDim / SUPER_RESOLUTION_FACTOR, -sqrt(2.0f) / VLMVoxelizationParams.VolumeMaxDim / SUPER_RESOLUTION_FACTOR, 0.0f},
+		{-sqrt(2.0f) / VLMVoxelizationParams.VolumeMaxDim / SUPER_RESOLUTION_FACTOR,  sqrt(2.0f) / VLMVoxelizationParams.VolumeMaxDim / SUPER_RESOLUTION_FACTOR, 0.0f},
+		{-sqrt(2.0f) / VLMVoxelizationParams.VolumeMaxDim / SUPER_RESOLUTION_FACTOR, -sqrt(2.0f) / VLMVoxelizationParams.VolumeMaxDim / SUPER_RESOLUTION_FACTOR, 0.0f}
 	};
 	
 	float MaxProj[3] = {0, 0, 0};
@@ -151,7 +152,7 @@ void VLMVoxelizationPS(
 {
 	// Recover {[-1, 1], [-1, 1], [-1, 1]} bounding box normalized coordinates from SV_Position (which is in {[0, ResX], [0, ResY], [0, 1]}
 	Input.ScreenPosition.xyz /= 0.5f;
-	Input.ScreenPosition.xy /= VLMVoxelizationParams.VolumeMaxDim;
+	Input.ScreenPosition.xy /= VLMVoxelizationParams.VolumeMaxDim * SUPER_RESOLUTION_FACTOR;
 	Input.ScreenPosition.x -= 1;
 	Input.ScreenPosition.y = 1 - Input.ScreenPosition.y;
 	Input.ScreenPosition.z -= 1;
@@ -175,8 +176,20 @@ void VLMVoxelizationPS(
 		WorldPosition.xyz = Input.ScreenPosition.xyz;
 	}
 
-	if (VLMVoxelizationParams.VoxelizeVolume[(WorldPosition + 1.0f) * 0.5f * VLMVoxelizationParams.VolumeMaxDim] == BRICK_IN_IMPORTANCE_VOLUME)
+	float3 VoxelPos = (WorldPosition + 1.0f) * 0.5f * VLMVoxelizationParams.VolumeMaxDim;
+
+	// Expand by 1 cell to avoid losing details from stitching to a higher mip
+	for (int dx = -1; dx <= 1; dx++)
 	{
-		VLMVoxelizationParams.VoxelizeVolume[(WorldPosition + 1.0f) * 0.5f * VLMVoxelizationParams.VolumeMaxDim] = BRICK_ALLOCATED;
+		for (int dy = -1; dy <= 1; dy++)
+		{
+			for (int dz = -1; dz <= 1; dz++)
+			{
+				if (VLMVoxelizationParams.VoxelizeVolume[VoxelPos + float3(dx, dy, dz) / SUPER_RESOLUTION_FACTOR] == BRICK_IN_IMPORTANCE_VOLUME)
+				{
+					VLMVoxelizationParams.VoxelizeVolume[VoxelPos + float3(dx, dy, dz) / SUPER_RESOLUTION_FACTOR] = BRICK_ALLOCATED;
+				}
+			}
+		}
 	}
 }
diff --git a/Engine/Plugins/Experimental/GPULightmass/Source/GPULightmass/Private/VolumetricLightmap.cpp b/Engine/Plugins/Experimental/GPULightmass/Source/GPULightmass/Private/VolumetricLightmap.cpp
--- a/Engine/Plugins/Experimental/GPULightmass/Source/GPULightmass/Private/VolumetricLightmap.cpp
+++ b/Engine/Plugins/Experimental/GPULightmass/Source/GPULightmass/Private/VolumetricLightmap.cpp
@@ -18,7 +18,6 @@ IMPLEMENT_MATERIAL_SHADER_TYPE(, FVLMVoxelizationGS, TEXT("/Plugin/GPULightmass/
 IMPLEMENT_MATERIAL_SHADER_TYPE(, FVLMVoxelizationPS, TEXT("/Plugin/GPULightmass/Private/VolumetricLightmapVoxelization.usf"), TEXT("VLMVoxelizationPS"), SF_Pixel);
 IMPLEMENT_GLOBAL_SHADER(FClearVolumeCS, "/Plugin/GPULightmass/Private/BrickAllocationManagement.usf", "ClearVolumeCS", SF_Compute);
 IMPLEMENT_GLOBAL_SHADER(FVoxelizeImportanceVolumeCS, "/Plugin/GPULightmass/Private/BrickAllocationManagement.usf", "VoxelizeImportanceVolumeCS", SF_Compute);
-IMPLEMENT_GLOBAL_SHADER(FDilateVolumeCS, "/Plugin/GPULightmass/Private/BrickAllocationManagement.usf", "DilateVolumeCS", SF_Compute);
 IMPLEMENT_GLOBAL_SHADER(FDownsampleVolumeCS, "/Plugin/GPULightmass/Private/BrickAllocationManagement.usf", "DownsampleVolumeCS", SF_Compute);
 IMPLEMENT_GLOBAL_SHADER(FCountNumBricksCS, "/Plugin/GPULightmass/Private/BrickAllocationManagement.usf", "CountNumBricksCS", SF_Compute);
 IMPLEMENT_GLOBAL_SHADER(FGatherBrickRequestsCS, "/Plugin/GPULightmass/Private/BrickAllocationManagement.usf", "GatherBrickRequestsCS", SF_Compute);
@@ -247,8 +246,8 @@ void FVolumetricLightmapRenderer::VoxelizeScene()
 
 			FVoxelizeImportanceVolumeCS::FParameters* Parameters = GraphBuilder.AllocParameters<FVoxelizeImportanceVolumeCS::FParameters>();
 			Parameters->VolumeSize = VoxelizationVolumeMips[0]->GetDesc().GetSize();
-			Parameters->ImportanceVolumeMin = (FVector3f)ImportanceVolume.Min;
-			Parameters->ImportanceVolumeMax = (FVector3f)ImportanceVolume.Max;
+			Parameters->ImportanceVolumeMin = (FVector3f)ImportanceVolume.Min - FVector3f(2 * TargetDetailCellSize);
+			Parameters->ImportanceVolumeMax = (FVector3f)ImportanceVolume.Max + FVector3f(2 * TargetDetailCellSize);
 			Parameters->VLMVoxelizationParams = PassUniformBuffer;
 			Parameters->VoxelizeVolume = VoxelizationVolumeMipUAVs[0];
 
@@ -277,7 +276,10 @@ void FVolumetricLightmapRenderer::VoxelizeScene()
 			ERDGPassFlags::Raster | ERDGPassFlags::NeverCull,
 			[this, CubeMaxDim, PassUniformBuffer, PassParameters](FRHICommandList& RHICmdList)
 		{
-			RHICmdList.SetViewport(0, 0, 0, CubeMaxDim, CubeMaxDim, 1);
+			// Must match VolumetricLightmapVoxelization.usf
+			const int32 SUPER_RESOLUTION_FACTOR = 4;
+			
+			RHICmdList.SetViewport(0, 0, 0, CubeMaxDim * SUPER_RESOLUTION_FACTOR, CubeMaxDim * SUPER_RESOLUTION_FACTOR, 1);
 
 			SCOPED_DRAW_EVENTF(RHICmdList, GPULightmassVoxelizeScene, TEXT("GPULightmass VoxelizeScene"));
 
@@ -368,21 +370,6 @@ void FVolumetricLightmapRenderer::VoxelizeScene()
 			});
 		});
 
-		{
-			TShaderMapRef<FDilateVolumeCS> ComputeShader(GlobalShaderMap);
-
-			FDilateVolumeCS::FParameters* Parameters = GraphBuilder.AllocParameters<FDilateVolumeCS::FParameters>();
-			Parameters->VolumeSize = VoxelizationVolumeMips[0]->GetDesc().GetSize();
-			Parameters->VoxelizeVolume = VoxelizationVolumeMipUAVs[0];
-
-			FComputeShaderUtils::AddPass(
-				GraphBuilder,
-				RDG_EVENT_NAME("DilateVolume"),
-				ComputeShader,
-				Parameters,
-				FComputeShaderUtils::GetGroupCount(VoxelizationVolumeMips[0]->GetDesc().GetSize(), FIntVector(4)));
-		}
-
 		for (int32 MipLevel = 1; MipLevel < VoxelizationVolumeMips.Num(); MipLevel++)
 		{
 			TShaderMapRef<FDownsampleVolumeCS> ComputeShader(GlobalShaderMap);
diff --git a/Engine/Plugins/Experimental/GPULightmass/Source/GPULightmass/Private/VolumetricLightmapVoxelization.h b/Engine/Plugins/Experimental/GPULightmass/Source/GPULightmass/Private/VolumetricLightmapVoxelization.h
--- a/Engine/Plugins/Experimental/GPULightmass/Source/GPULightmass/Private/VolumetricLightmapVoxelization.h
+++ b/Engine/Plugins/Experimental/GPULightmass/Source/GPULightmass/Private/VolumetricLightmapVoxelization.h
@@ -228,27 +228,6 @@ class FVoxelizeImportanceVolumeCS : public FGlobalShader
 	END_SHADER_PARAMETER_STRUCT()
 };
 
-class FDilateVolumeCS : public FGlobalShader
-{
-	DECLARE_GLOBAL_SHADER(FDilateVolumeCS);
-	SHADER_USE_PARAMETER_STRUCT(FDilateVolumeCS, FGlobalShader);
-
-	static bool ShouldCompilePermutation(const FGlobalShaderPermutationParameters& Parameters)
-	{
-		return EnumHasAllFlags(Parameters.Flags, EShaderPermutationFlags::HasEditorOnlyData) && RHISupportsRayTracingShaders(Parameters.Platform);
-	}
-
-	static void ModifyCompilationEnvironment(const FGlobalShaderPermutationParameters& Parameters, FShaderCompilerEnvironment& OutEnvironment)
-	{
-		OutEnvironment.CompilerFlags.Add(CFLAG_ForceDXC);
-	}
-
-	BEGIN_SHADER_PARAMETER_STRUCT(FParameters, )
-		SHADER_PARAMETER(FIntVector, VolumeSize)
-		SHADER_PARAMETER_RDG_TEXTURE_UAV(RWTexture3D<uint>, VoxelizeVolume)
-	END_SHADER_PARAMETER_STRUCT()
-};
-
 class FDownsampleVolumeCS : public FGlobalShader
 {
 	DECLARE_GLOBAL_SHADER(FDownsampleVolumeCS);
