GPULM: Evalutate both hemispheres for two sided materials. This fixes pitch black sides yet since our surface lighting representation isn't for full spheres, expect loss in directionality/overall brightness differences - but still better than pitch black

diff --git a/Engine/Plugins/Experimental/GPULightmass/Shaders/Private/LightmapGBuffer.usf b/Engine/Plugins/Experimental/GPULightmass/Shaders/Private/LightmapGBuffer.usf
--- a/Engine/Plugins/Experimental/GPULightmass/Shaders/Private/LightmapGBuffer.usf
+++ b/Engine/Plugins/Experimental/GPULightmass/Shaders/Private/LightmapGBuffer.usf
@@ -14,6 +14,7 @@
 float4 VirtualTexturePhysicalTileCoordinateScaleAndBias;
 int2 ScratchTilePoolOffset; 
 int RenderPassIndex;
+int bMaterialTwoSided;
 
 struct FLightmapGBufferVSToPS
 {
@@ -71,7 +72,7 @@ void LightmapGBufferPS(
 	FLightmapGBufferVSToPS Input
 	OPTIONAL_IsFrontFace
 	)
-{
+{ 
 	ResolvedView = ResolveView();
 	FMaterialPixelParameters MaterialParameters = GetMaterialPixelParameters(Input.FactoryInterpolants, Input.SavedWorldPosition);
 	FPixelMaterialInputs PixelMaterialInputs;
@@ -85,7 +86,7 @@ void LightmapGBufferPS(
     OutWorldNormal *= CondMask(bIsFrontFace, 1.0f, -1.0f);
 	OutWorldNormal *= GetPrimitive_DeterminantSign(MaterialParameters.PrimitiveId);
 	// Interpolated vertex normal
-	float4 OutShadingNormal = float4(MaterialParameters.TangentToWorld[2], 0);
+	float4 OutShadingNormal = float4(MaterialParameters.TangentToWorld[2], bMaterialTwoSided);
 
 	// GBuffer Outputs
 	LightmapGBufferParams.ScratchTilePoolLayer0[ScratchTilePoolOffset + Input.Position.xy] = OutWorldPosition;
diff --git a/Engine/Plugins/Experimental/GPULightmass/Shaders/Private/LightmapPathTracing.usf b/Engine/Plugins/Experimental/GPULightmass/Shaders/Private/LightmapPathTracing.usf
--- a/Engine/Plugins/Experimental/GPULightmass/Shaders/Private/LightmapPathTracing.usf
+++ b/Engine/Plugins/Experimental/GPULightmass/Shaders/Private/LightmapPathTracing.usf
@@ -875,6 +875,7 @@ void LightmapPathTracingMainRG()
 
 	float3 WorldNormal = GBufferWorldNormal[TexelIndexInScratch].xyz;
 	float3 ShadingNormal = GBufferShadingNormal[TexelIndexInScratch].xyz;
+	bool bMaterialTwoSided = GBufferShadingNormal[TexelIndexInScratch].w;
 
 	float4 EncodedGBufferWorldPosition = GBufferWorldPosition[TexelIndexInScratch];
 
@@ -917,7 +918,9 @@ void LightmapPathTracingMainRG()
 			uint Seed = VirtualTextureSpacePosition.y * BatchedTiles[TileIndex].LightmapSize.x + VirtualTextureSpacePosition.x;
 					
 			float4 PrimaryRandSample = float4(-1, -1, -1, -1);
-			
+
+			// Evaluate both hemispheres for two sided materials and add the lighting up naively
+			for (int HemisphereIndex = 0; HemisphereIndex < (bMaterialTwoSided ? 2 : 1); HemisphereIndex++)
 			{
 				RandomSequence RandSequence;
 				RandomSequence_Initialize(RandSequence, Seed, EffectiveRenderPassIndex);
@@ -934,13 +937,15 @@ void LightmapPathTracingMainRG()
 			#endif
 				float LuminanceForFirstBounceRayGuiding = 0;
 				float3 SkyLightBentNormal = 0;
+
+				float3 EffectiveShadingNormal = HemisphereIndex == 0 ? ShadingNormal : -ShadingNormal;
 				
 				if (bIsValidSample)
 				{
 					PathTracingKernel(
 						BatchedTiles[TileIndex].RenderPassIndex,
 						TranslatedWorldPosition,
-						ShadingNormal,
+						EffectiveShadingNormal,
 						RandSequence,
 						SampleIndex,
 						bIsValidSample,
@@ -979,7 +984,7 @@ void LightmapPathTracingMainRG()
 						int2 FinalPosition = ClusterPosition * DIRECTIONAL_BINS_ONE_DIM + PositionInBin;
 
 						{
-							float Illuminance = LuminanceForFirstBounceRayGuiding * saturate(dot(RadianceDirection, ShadingNormal));
+							float Illuminance = LuminanceForFirstBounceRayGuiding * saturate(dot(RadianceDirection, EffectiveShadingNormal));
 							// For positive floats we can cast them to uint and use atomic max directly
 							InterlockedMax(RayGuidingLuminance[BatchedTiles[TileIndex].WorkingSetPosition / GPreviewLightmapPhysicalTileSize * CDF_TILE_SIZE + FinalPosition], asuint(max(Illuminance, 0)));
 						}
@@ -989,7 +994,7 @@ void LightmapPathTracingMainRG()
 					{
 						FL2SHAndCorrection SH = (FL2SHAndCorrection)0;
 						
-						float TangentZ = saturate(dot(RadianceDirection, ShadingNormal));
+						float TangentZ = saturate(dot(RadianceDirection, EffectiveShadingNormal));
 						SH.AddIncomingRadiance(Luminance(RadianceValue), RadianceDirection, TangentZ);
 						
 						IrradianceAndSampleCount[TexelIndexInPool].rgb += float3(RadianceValue * TangentZ / PI);
diff --git a/Engine/Plugins/Experimental/GPULightmass/Source/GPULightmass/Private/LightmapGBuffer.h b/Engine/Plugins/Experimental/GPULightmass/Source/GPULightmass/Private/LightmapGBuffer.h
--- a/Engine/Plugins/Experimental/GPULightmass/Source/GPULightmass/Private/LightmapGBuffer.h
+++ b/Engine/Plugins/Experimental/GPULightmass/Source/GPULightmass/Private/LightmapGBuffer.h
@@ -23,6 +23,7 @@ struct FLightmapElementData : public FMeshMaterialShaderElementData
 	FVector4f VirtualTexturePhysicalTileCoordinateScaleAndBias;
 	int32 RenderPassIndex;
 	FIntPoint ScratchTilePoolOffset;
+	int32 bMaterialTwoSided = 0;
 
 	FLightmapElementData(const FLightCacheInterface* LCI) : LCI(LCI) {}
 };
@@ -110,6 +111,7 @@ class FLightmapGBufferPS : public FMeshMaterialShader
 	DECLARE_SHADER_TYPE(FLightmapGBufferPS, MeshMaterial);
 
 	LAYOUT_FIELD(FShaderParameter, ScratchTilePoolOffset)
+	LAYOUT_FIELD(FShaderParameter, bMaterialTwoSided)
 public:
 	static bool ShouldCompilePermutation(const FMeshMaterialShaderPermutationParameters& Parameters)
 	{
@@ -145,6 +147,7 @@ public:
 		: FMeshMaterialShader(Initializer)
 	{
 		ScratchTilePoolOffset.Bind(Initializer.ParameterMap, TEXT("ScratchTilePoolOffset"));
+		bMaterialTwoSided.Bind(Initializer.ParameterMap, TEXT("bMaterialTwoSided"));
 	}
 
 	void GetShaderBindings(
@@ -159,6 +162,7 @@ public:
 	{
 		FMeshMaterialShader::GetShaderBindings(Scene, FeatureLevel, PrimitiveSceneProxy, MaterialRenderProxy, Material, DrawRenderState, ShaderElementData, ShaderBindings);
 		ShaderBindings.Add(ScratchTilePoolOffset, ShaderElementData.ScratchTilePoolOffset);
+		ShaderBindings.Add(bMaterialTwoSided, ShaderElementData.bMaterialTwoSided);
 	}
 };
 
@@ -186,10 +190,16 @@ public:
 	{
 		if (MeshBatch.bUseForMaterial)
 		{
+			bool bMaterialTwoSided = false;
+			if (MeshBatch.MaterialRenderProxy && MeshBatch.MaterialRenderProxy->GetMaterialNoFallback(FeatureLevel))
+			{
+				bMaterialTwoSided = MeshBatch.MaterialRenderProxy->GetMaterialNoFallback(FeatureLevel)->IsTwoSided();
+			}
+			
 			const FMaterialRenderProxy& DefaultProxy = *UMaterial::GetDefaultMaterial(MD_Surface)->GetRenderProxy();
 			const FMaterial& DefaultMaterial = *DefaultProxy.GetMaterialNoFallback(FeatureLevel);
 
-			Process(MeshBatch, BatchElementMask, StaticMeshId, PrimitiveSceneProxy, DefaultProxy, DefaultMaterial);
+			Process(MeshBatch, BatchElementMask, StaticMeshId, bMaterialTwoSided, PrimitiveSceneProxy, DefaultProxy, DefaultMaterial);
 		}
 	}
 
@@ -198,6 +208,7 @@ private:
 		const FMeshBatch& MeshBatch,
 		uint64 BatchElementMask,
 		int32 StaticMeshId,
+		bool bMaterialTwoSided,
 		const FPrimitiveSceneProxy* RESTRICT PrimitiveSceneProxy,
 		const FMaterialRenderProxy& RESTRICT MaterialRenderProxy,
 		const FMaterial& RESTRICT MaterialResource)
@@ -227,6 +238,7 @@ private:
 		ShaderElementData.VirtualTexturePhysicalTileCoordinateScaleAndBias = VirtualTexturePhysicalTileCoordinateScaleAndBias;
 		ShaderElementData.RenderPassIndex = RenderPassIndex;
 		ShaderElementData.ScratchTilePoolOffset = ScratchTilePoolOffset;
+		ShaderElementData.bMaterialTwoSided = bMaterialTwoSided;
 
 		FMeshDrawCommandSortKey SortKey {};
 
