GPULM: Fix black volumetric lightmaps when irradiance caching is enabled

diff --git a/Engine/Plugins/Experimental/GPULightmass/Source/GPULightmass/Private/VolumetricLightmap.cpp b/Engine/Plugins/Experimental/GPULightmass/Source/GPULightmass/Private/VolumetricLightmap.cpp
--- a/Engine/Plugins/Experimental/GPULightmass/Source/GPULightmass/Private/VolumetricLightmap.cpp
+++ b/Engine/Plugins/Experimental/GPULightmass/Source/GPULightmass/Private/VolumetricLightmap.cpp
@@ -867,6 +867,12 @@ void FVolumetricLightmapRenderer::BackgroundTick()
 			PassParameters->NumTotalBricks = NumTotalBricks;
 			PassParameters->BrickBatchOffset = BrickBatchOffset;
 			PassParameters->NumTotalPassesToRender = NumTotalPassesToRender;
+			
+			if (Scene->Settings->bUseIrradianceCaching)
+			{
+				PassParameters->NumTotalPassesToRender -= Scene->Settings->IrradianceCacheQuality;	
+			}
+			
 			PassParameters->BrickRequests = BrickRequests.UAV;
 			PassParameters->AmbientVector = AccumulationBrickDataRDG.AmbientVectorRDG;
 			PassParameters->SHCoefficients0R = AccumulationBrickDataRDG.SHCoefficientsRDG[0];
