GPULM: Fix issues with empty scenes

diff --git a/Engine/Plugins/Experimental/GPULightmass/Source/GPULightmass/Private/Scene/Lights.cpp b/Engine/Plugins/Experimental/GPULightmass/Source/GPULightmass/Private/Scene/Lights.cpp
--- a/Engine/Plugins/Experimental/GPULightmass/Source/GPULightmass/Private/Scene/Lights.cpp
+++ b/Engine/Plugins/Experimental/GPULightmass/Source/GPULightmass/Private/Scene/Lights.cpp
@@ -364,7 +364,10 @@ void FSkyLightRenderState::PrepareSkyTexture(FRHICommandListImmediate& RHICmdLis
 
 void FDirectionalLightRenderState::RenderStaticShadowDepthMap(FRHICommandListImmediate& RHICmdList, FSceneRenderState& Scene)
 {
-	Scene.SetupRayTracingScene();
+	if (!Scene.SetupRayTracingScene())
+	{
+		return;
+	}
 	
 	FVector3f XAxis, YAxis;
 	FVector3f(Direction).FindBestAxisVectors(XAxis, YAxis);
@@ -461,7 +464,10 @@ void FDirectionalLightRenderState::RenderStaticShadowDepthMap(FRHICommandListImm
 
 void FSpotLightRenderState::RenderStaticShadowDepthMap(FRHICommandListImmediate& RHICmdList, FSceneRenderState& Scene)
 {
-	Scene.SetupRayTracingScene();
+	if (!Scene.SetupRayTracingScene())
+	{
+		return;
+	}
 	
 	FVector3f XAxis, YAxis;
 	FVector3f(Direction).FindBestAxisVectors(XAxis, YAxis);
@@ -577,7 +583,10 @@ void FSpotLightRenderState::RenderStaticShadowDepthMap(FRHICommandListImmediate&
 
 void FPointLightRenderState::RenderStaticShadowDepthMap(FRHICommandListImmediate& RHICmdList, FSceneRenderState& Scene)
 {
-	Scene.SetupRayTracingScene();
+	if (!Scene.SetupRayTracingScene())
+	{
+		return;
+	}
 
 	const float ClampedResolutionScale = 1.0f;
 	const float StaticShadowDepthMapTransitionSampleDistanceX = 10;
@@ -667,7 +676,10 @@ void FPointLightRenderState::RenderStaticShadowDepthMap(FRHICommandListImmediate
 
 void FRectLightRenderState::RenderStaticShadowDepthMap(FRHICommandListImmediate& RHICmdList, FSceneRenderState& Scene)
 {
-	Scene.SetupRayTracingScene();
+	if (!Scene.SetupRayTracingScene())
+	{
+		return;
+	}
 
 	const float ClampedResolutionScale = 1.0f;
 	const float StaticShadowDepthMapTransitionSampleDistanceX = 10;
diff --git a/Engine/Plugins/Experimental/GPULightmass/Source/GPULightmass/Private/Scene/Scene.cpp b/Engine/Plugins/Experimental/GPULightmass/Source/GPULightmass/Private/Scene/Scene.cpp
--- a/Engine/Plugins/Experimental/GPULightmass/Source/GPULightmass/Private/Scene/Scene.cpp
+++ b/Engine/Plugins/Experimental/GPULightmass/Source/GPULightmass/Private/Scene/Scene.cpp
@@ -1637,7 +1637,14 @@ void FSceneRenderState::BackgroundTick()
 		int32 IntPercentage = FMath::FloorToInt(SamplesTaken * 100.0 / TotalSamples);
 		IntPercentage = FMath::Max(IntPercentage, 0);
 		// With high number of samples (like 8192) double precision isn't enough to prevent fake 100%s
-		IntPercentage = FMath::Min(IntPercentage, SamplesTaken < TotalSamples ? 99 : 100);
+		if (SamplesTaken < TotalSamples)
+		{
+			IntPercentage = FMath::Min(IntPercentage, 99);
+		}
+		else
+		{
+			IntPercentage = 100;
+		}
 
 		FPlatformAtomics::InterlockedExchange(&Percentage, IntPercentage);
 	}
@@ -1915,6 +1922,7 @@ void FScene::ApplyFinishedLightmapsToWorld()
 			RectLight.AllocateMapBuildData(LightingScenario ? LightingScenario : RectLight.ComponentUObject->GetOwner()->GetLevel());
 		}
 
+		if (RenderState.VolumetricLightmapRenderer->NumTotalBricks > 0)
 		{
 			ULevel* SubLevelStorageLevel = LightingScenario ? LightingScenario : ToRawPtr(World->PersistentLevel);
 			UMapBuildDataRegistry* SubLevelRegistry = SubLevelStorageLevel->GetOrCreateMapBuildData();
