GPULM: Use double precision matrices to avoid NaNs with LWC. No actual support for LWC at the moment.

diff --git a/Engine/Plugins/Experimental/GPULightmass/Source/GPULightmass/Private/Scene/Lights.cpp b/Engine/Plugins/Experimental/GPULightmass/Source/GPULightmass/Private/Scene/Lights.cpp
--- a/Engine/Plugins/Experimental/GPULightmass/Source/GPULightmass/Private/Scene/Lights.cpp
+++ b/Engine/Plugins/Experimental/GPULightmass/Source/GPULightmass/Private/Scene/Lights.cpp
@@ -369,12 +369,12 @@ void FDirectionalLightRenderState::RenderStaticShadowDepthMap(FRHICommandListImm
 		return;
 	}
 	
-	FVector3f XAxis, YAxis;
-	FVector3f(Direction).FindBestAxisVectors(XAxis, YAxis);
+	FVector XAxis, YAxis;
+	Direction.FindBestAxisVectors(XAxis, YAxis);
 	// Create a coordinate system for the dominant directional light, with the z axis corresponding to the light's direction
-	FMatrix44f WorldToLight = FBasisVectorMatrix44f(XAxis, YAxis, FVector3f(Direction), FVector4f(0, 0, 0));
+	FMatrix WorldToLight = FBasisVectorMatrix(XAxis, YAxis, Direction, FVector4(0, 0, 0));
 
-	const FBox3f LightSpaceImportanceBounds = FBox3f{Scene.CombinedImportanceVolume}.TransformBy(WorldToLight);
+	const FBox LightSpaceImportanceBounds = Scene.CombinedImportanceVolume.TransformBy(WorldToLight);
 
 	const float ClampedResolutionScale = 1.0f;
 	const float StaticShadowDepthMapTransitionSampleDistanceX = 10;
@@ -408,9 +408,9 @@ void FDirectionalLightRenderState::RenderStaticShadowDepthMap(FRHICommandListImm
 	PassParameters->TLAS = Scene.RayTracingSceneSRV;
 	PassParameters->ShadowMapSize = FIntPoint{ShadowMapSizeX, ShadowMapSizeY};
 	PassParameters->StaticShadowDepthMapSuperSampleFactor = 1;
-	PassParameters->LightSpaceImportanceBoundsMin = LightSpaceImportanceBounds.Min;
-	PassParameters->LightSpaceImportanceBoundsMax = LightSpaceImportanceBounds.Max;
-	PassParameters->LightToWorld = WorldToLight.InverseFast();
+	PassParameters->LightSpaceImportanceBoundsMin = FVector3f(LightSpaceImportanceBounds.Min);
+	PassParameters->LightSpaceImportanceBoundsMax = FVector3f(LightSpaceImportanceBounds.Max);
+	PassParameters->LightToWorld = FMatrix44f(WorldToLight.InverseFast());
 	PassParameters->MaxPossibleDistance = LightSpaceImportanceBounds.Max.Z - LightSpaceImportanceBounds.Min.Z;
 	PassParameters->DepthMapTexture = GraphBuilder.CreateUAV(DepthMapTexture);
 
@@ -455,8 +455,8 @@ void FDirectionalLightRenderState::RenderStaticShadowDepthMap(FRHICommandListImm
 	}
 	LightComponentMapBuildData->DepthMap.ShadowMapSizeX = ShadowMapSizeX;
 	LightComponentMapBuildData->DepthMap.ShadowMapSizeY = ShadowMapSizeY;
-	WorldToLight *= FTranslationMatrix44f(-LightSpaceImportanceBounds.Min) * FScaleMatrix44f(FVector3f(1.0f) / (LightSpaceImportanceBounds.Max - LightSpaceImportanceBounds.Min));
-	LightComponentMapBuildData->DepthMap.WorldToLight = FMatrix(WorldToLight);
+	WorldToLight *= FTranslationMatrix(-LightSpaceImportanceBounds.Min) * FScaleMatrix(FVector(1.0) / (LightSpaceImportanceBounds.Max - LightSpaceImportanceBounds.Min));
+	LightComponentMapBuildData->DepthMap.WorldToLight = WorldToLight;
 	DepthMapTextureReadback.Unlock();
 	
 	Scene.DestroyRayTracingScene();
@@ -469,18 +469,18 @@ void FSpotLightRenderState::RenderStaticShadowDepthMap(FRHICommandListImmediate&
 		return;
 	}
 	
-	FVector3f XAxis, YAxis;
-	FVector3f(Direction).FindBestAxisVectors(XAxis, YAxis);
+	FVector XAxis, YAxis;
+	Direction.FindBestAxisVectors(XAxis, YAxis);
 	// Create a coordinate system for the dominant directional light, with the z axis corresponding to the light's direction
-	FMatrix44f WorldToLight = FTranslationMatrix44f(-FVector3f(Position)) * FBasisVectorMatrix44f(XAxis, YAxis, FVector3f(Direction), FVector4f(0, 0, 0));
+	FMatrix WorldToLight = FTranslationMatrix(-Position) * FBasisVectorMatrix(XAxis, YAxis, Direction, FVector4(0, 0, 0));
 
 	// Distance from the light's direction axis to the edge of the cone at the radius of the light
 	const float HalfCrossSectionLength = AttenuationRadius * FMath::Tan(FMath::Acos(SpotAngles.X));
 	
-	const FBox3f LightSpaceImportanceBounds
+	const FBox LightSpaceImportanceBounds
 	{
-		FVector3f{-HalfCrossSectionLength, -HalfCrossSectionLength, 0},
-		FVector3f{HalfCrossSectionLength, HalfCrossSectionLength, AttenuationRadius},
+		FVector{-HalfCrossSectionLength, -HalfCrossSectionLength, 0},
+		FVector{HalfCrossSectionLength, HalfCrossSectionLength, AttenuationRadius},
 	};
 	
 	const float ClampedResolutionScale = 1.0f;
@@ -511,17 +511,17 @@ void FSpotLightRenderState::RenderStaticShadowDepthMap(FRHICommandListImmediate&
 		), TEXT("GPULMStaticShadowDepthMap"));
 
 	const float MaxPossibleDistance = LightSpaceImportanceBounds.Max.Z - LightSpaceImportanceBounds.Min.Z;
-	const FMatrix44f LightToWorld = WorldToLight.InverseFast();
+	const FMatrix LightToWorld = WorldToLight.InverseFast();
 
 	FStaticShadowDepthMapTracingRGS::FParameters* PassParameters = GraphBuilder.AllocParameters<FStaticShadowDepthMapTracingRGS::FParameters>();
 	PassParameters->ViewUniformBuffer = Scene.ReferenceView->ViewUniformBuffer;
 	PassParameters->TLAS = Scene.RayTracingSceneSRV;
 	PassParameters->ShadowMapSize = FIntPoint{ShadowMapSizeX, ShadowMapSizeY};
 	PassParameters->StaticShadowDepthMapSuperSampleFactor = 1;
-	PassParameters->LightSpaceImportanceBoundsMin = LightSpaceImportanceBounds.Min;
-	PassParameters->LightSpaceImportanceBoundsMax = LightSpaceImportanceBounds.Max;
-	PassParameters->LightToWorld = LightToWorld;
-	PassParameters->WorldToLight = WorldToLight;
+	PassParameters->LightSpaceImportanceBoundsMin = FVector3f(LightSpaceImportanceBounds.Min);
+	PassParameters->LightSpaceImportanceBoundsMax = FVector3f(LightSpaceImportanceBounds.Max);
+	PassParameters->LightToWorld = FMatrix44f(LightToWorld);
+	PassParameters->WorldToLight = FMatrix44f(WorldToLight);
 	PassParameters->MaxPossibleDistance = MaxPossibleDistance;
 	PassParameters->DepthMapTexture = GraphBuilder.CreateUAV(DepthMapTexture);
 
@@ -568,14 +568,14 @@ void FSpotLightRenderState::RenderStaticShadowDepthMap(FRHICommandListImmediate&
 	LightComponentMapBuildData->DepthMap.ShadowMapSizeY = ShadowMapSizeY;
 	WorldToLight *=
 		// Perspective projection sized to the spotlight cone
-		FPerspectiveMatrix44f(FMath::Acos(SpotAngles.X), 1, 1, 0, AttenuationRadius)
+		FPerspectiveMatrix(FMath::Acos(SpotAngles.X), 1, 1, 0, AttenuationRadius)
 		// Convert from NDC to texture space, normalize Z
-		* FMatrix44f(
-			FPlane4f(.5f,	0,		0,											0),
-			FPlane4f(0,	.5f,	0,											0),
-			FPlane4f(0,	0,		1.0f / LightSpaceImportanceBounds.Max.Z,	0),
-			FPlane4f(.5f,	.5f,	0,											1));
-	LightComponentMapBuildData->DepthMap.WorldToLight = FMatrix(WorldToLight);
+		* FMatrix(
+			FPlane(.5,	0,		0,											0),
+			FPlane(0,	.5,	0,											0),
+			FPlane(0,	0,		1.0 / LightSpaceImportanceBounds.Max.Z,	0),
+			FPlane(.5,	.5,	0,											1));
+	LightComponentMapBuildData->DepthMap.WorldToLight = WorldToLight;
 	DepthMapTextureReadback.Unlock();
 	
 	Scene.DestroyRayTracingScene();
