Fix a few more shaders that crash dxc when optimization is disabled (UE-165154).

diff --git a/Engine/Plugins/Experimental/GPULightmass/Source/GPULightmass/Private/LightmapRayTracing.h b/Engine/Plugins/Experimental/GPULightmass/Source/GPULightmass/Private/LightmapRayTracing.h
--- a/Engine/Plugins/Experimental/GPULightmass/Source/GPULightmass/Private/LightmapRayTracing.h
+++ b/Engine/Plugins/Experimental/GPULightmass/Source/GPULightmass/Private/LightmapRayTracing.h
@@ -64,6 +64,8 @@ class FLightmapPathTracingRGS : public FGlobalShader
 		OutEnvironment.SetDefine(TEXT("USE_RECT_LIGHT_TEXTURES"), 1);
 		OutEnvironment.CompilerFlags.Add(CFLAG_ForceDXC);
 		OutEnvironment.CompilerFlags.Add(CFLAG_WarningsAsErrors);
+		// @todo - Working around DXC compiler crash UE-165154, should be removed after DXC is updated with a fix.
+		OutEnvironment.CompilerFlags.Add(CFLAG_ForceOptimization);
 	}
 
 	BEGIN_SHADER_PARAMETER_STRUCT(FParameters, )
@@ -123,6 +125,8 @@ class FVolumetricLightmapPathTracingRGS : public FGlobalShader
 		OutEnvironment.SetDefine(TEXT("USE_RECT_LIGHT_TEXTURES"), 1);
 		OutEnvironment.CompilerFlags.Add(CFLAG_ForceDXC);
 		OutEnvironment.CompilerFlags.Add(CFLAG_WarningsAsErrors);
+		// @todo - Working around DXC compiler crash UE-165154, should be removed after DXC is updated with a fix.
+		OutEnvironment.CompilerFlags.Add(CFLAG_ForceOptimization);
 	}
 
 	BEGIN_SHADER_PARAMETER_STRUCT(FParameters, )
