Add new GetActorPositionForRenderer() for use by RendererScene.
This function is used to get the ActorPosition stored in the scene, and used by material expressions.
This is expected to be the position of the actor at the root of the component hierachy.
Added a component flag that allows us to early out of the walk of the component hierachy.
Set the flag on the ALevelInstanceEditorInstanceActor root component so that these actors are ignored.
This fixes a bug where Level Instances add this actor into the hierachy only in the editor, meaning that the results of the material rendering would change between editor and game.

diff --git a/Engine/Plugins/Experimental/GPULightmass/Source/GPULightmass/Private/Scene/Scene.cpp b/Engine/Plugins/Experimental/GPULightmass/Source/GPULightmass/Private/Scene/Scene.cpp
--- a/Engine/Plugins/Experimental/GPULightmass/Source/GPULightmass/Private/Scene/Scene.cpp
+++ b/Engine/Plugins/Experimental/GPULightmass/Source/GPULightmass/Private/Scene/Scene.cpp
@@ -746,7 +746,7 @@ void FScene::AddGeometryInstanceFromComponent(UStaticMeshComponent* InComponent)
 	InstanceRenderState.RenderData = Instance->ComponentUObject->GetStaticMesh()->GetRenderData();
 	InstanceRenderState.LocalToWorld = InComponent->GetRenderMatrix();
 	InstanceRenderState.WorldBounds = InComponent->Bounds;
-	InstanceRenderState.ActorPosition = InComponent->GetAttachmentRootActor() ? InComponent->GetAttachmentRootActor()->GetActorLocation() : FVector(ForceInitToZero);
+	InstanceRenderState.ActorPosition = InComponent->GetActorPositionForRenderer();
 	InstanceRenderState.LocalBounds = InComponent->CalcBounds(FTransform::Identity);
 	InstanceRenderState.bCastShadow = InComponent->CastShadow && InComponent->bCastStaticShadow;
 	InstanceRenderState.LODOverrideColorVertexBuffers.AddZeroed(InComponent->GetStaticMesh()->GetRenderData()->LODResources.Num());
@@ -1017,7 +1017,7 @@ void FScene::AddGeometryInstanceFromComponent(UInstancedStaticMeshComponent* InC
 	InstanceRenderState.InstancedRenderData = MakeUnique<FInstancedStaticMeshRenderData>(Instance->ComponentUObject, FeatureLevel);
 	InstanceRenderState.LocalToWorld = InComponent->GetRenderMatrix();
 	InstanceRenderState.WorldBounds = InComponent->Bounds;
-	InstanceRenderState.ActorPosition = InComponent->GetAttachmentRootActor() ? InComponent->GetAttachmentRootActor()->GetActorLocation() : FVector(ForceInitToZero);
+	InstanceRenderState.ActorPosition = InComponent->GetActorPositionForRenderer();
 	InstanceRenderState.LocalBounds = InComponent->CalcBounds(FTransform::Identity);
 	InstanceRenderState.bCastShadow = InComponent->CastShadow && InComponent->bCastStaticShadow;
 
@@ -1196,7 +1196,7 @@ void FScene::AddGeometryInstanceFromComponent(ULandscapeComponent* InComponent)
 	InstanceRenderState.ComponentUObject = Instance->ComponentUObject;
 	InstanceRenderState.LocalToWorld = InComponent->GetRenderMatrix();
 	InstanceRenderState.WorldBounds = InComponent->Bounds;
-	InstanceRenderState.ActorPosition = InComponent->GetAttachmentRootActor() ? InComponent->GetAttachmentRootActor()->GetActorLocation() : FVector(ForceInitToZero);
+	InstanceRenderState.ActorPosition = InComponent->GetActorPositionForRenderer();
 	InstanceRenderState.LocalBounds = InComponent->CalcBounds(FTransform::Identity);
 	InstanceRenderState.bCastShadow = InComponent->CastShadow && InComponent->bCastStaticShadow;
 
diff --git a/Engine/Source/Runtime/Engine/Classes/Components/SceneComponent.h b/Engine/Source/Runtime/Engine/Classes/Components/SceneComponent.h
--- a/Engine/Source/Runtime/Engine/Classes/Components/SceneComponent.h
+++ b/Engine/Source/Runtime/Engine/Classes/Components/SceneComponent.h
@@ -255,6 +255,10 @@ public:
 	UFUNCTION(BlueprintSetter)
 	void SetShouldUpdatePhysicsVolume(bool bInShouldUpdatePhysicsVolume);
 
+	/** If true, this component stops the the walk up the attachment chain in GetActorPositionForRenderer(). Instead this component's child will be used as the attachment root. */
+	UPROPERTY()
+	uint8 bIsNotRenderAttachmentRoot : 1;
+
 protected:
 	/** Transient flag that temporarily disables UpdateOverlaps within DetachFromParent(). */
 	uint8 bDisableDetachmentUpdateOverlaps:1;
@@ -1246,6 +1250,9 @@ public:
 	/** Walks up the attachment chain from this SceneComponent and returns the top-level actor it's attached to.  Returns Owner if unattached. */
 	AActor* GetAttachmentRootActor() const;
 
+	/** Returns the ActorPosition for use by rendering. This comes from walking the attachment chain to the top-level actor. */
+	FVector GetActorPositionForRenderer() const;
+
 	/** Gets the owner of the attach parent */
 	AActor* GetAttachParentActor() const;
 
diff --git a/Engine/Source/Runtime/Engine/Private/Components/SceneComponent.cpp b/Engine/Source/Runtime/Engine/Private/Components/SceneComponent.cpp
--- a/Engine/Source/Runtime/Engine/Private/Components/SceneComponent.cpp
+++ b/Engine/Source/Runtime/Engine/Private/Components/SceneComponent.cpp
@@ -2360,6 +2360,13 @@ AActor* USceneComponent::GetAttachmentRootActor() const
 	return AttachmentRootComponent ? AttachmentRootComponent->GetOwner() : nullptr;
 }
 
+FVector USceneComponent::GetActorPositionForRenderer() const
+{
+	const USceneComponent* Top;
+	for (Top = this; Top->GetAttachParent() && !Top->GetAttachParent()->bIsNotRenderAttachmentRoot; Top = Top->GetAttachParent());
+	return (Top->GetOwner() != nullptr) ? Top->GetOwner()->GetActorLocation() : FVector(ForceInitToZero);
+}
+
 AActor* USceneComponent::GetAttachParentActor() const
 {
 	const USceneComponent* const AttachParentComponent = GetAttachParent();
diff --git a/Engine/Source/Runtime/Engine/Private/LevelInstance/LevelInstanceEditorInstanceActor.cpp b/Engine/Source/Runtime/Engine/Private/LevelInstance/LevelInstanceEditorInstanceActor.cpp
--- a/Engine/Source/Runtime/Engine/Private/LevelInstance/LevelInstanceEditorInstanceActor.cpp
+++ b/Engine/Source/Runtime/Engine/Private/LevelInstance/LevelInstanceEditorInstanceActor.cpp
@@ -12,6 +12,10 @@ ALevelInstanceEditorInstanceActor::ALevelInstanceEditorInstanceActor(const FObje
 {
 	RootComponent = CreateDefaultSubobject<USceneComponent>(TEXT("RootComponent"));
 	RootComponent->Mobility = EComponentMobility::Static;
+	
+	// To keep the behavior of any calls to USceneComponent::GetActorPositionForRenderer() consistent between Editor and Game modes
+	// we need to flag the root component so that it isn't considered as an AttachmentRoot.
+	RootComponent->bIsNotRenderAttachmentRoot = true;
 }
 
 #if WITH_EDITOR
diff --git a/Engine/Source/Runtime/Renderer/Private/RendererScene.cpp b/Engine/Source/Runtime/Renderer/Private/RendererScene.cpp
--- a/Engine/Source/Runtime/Renderer/Private/RendererScene.cpp
+++ b/Engine/Source/Runtime/Renderer/Private/RendererScene.cpp
@@ -2036,13 +2036,7 @@ void FScene::AddPrimitive(UPrimitiveComponent* Primitive)
 
 	// Cache the primitives initial transform.
 	FMatrix RenderMatrix = Primitive->GetRenderMatrix();
-	FVector AttachmentRootPosition(0);
-
-	AActor* AttachmentRoot = Primitive->GetAttachmentRootActor();
-	if (AttachmentRoot)
-	{
-		AttachmentRootPosition = AttachmentRoot->GetActorLocation();
-	}
+	FVector AttachmentRootPosition = Primitive->GetActorPositionForRenderer();
 
 	struct FCreateRenderThreadParameters
 	{
@@ -2185,13 +2179,7 @@ void FScene::UpdatePrimitiveTransform(UPrimitiveComponent* Primitive)
 		}
 		else
 		{
-			FVector AttachmentRootPosition(0);
-
-			AActor* Actor = Primitive->GetAttachmentRootActor();
-			if (Actor != nullptr)
-			{
-				AttachmentRootPosition = Actor->GetActorLocation();
-			}
+			FVector AttachmentRootPosition = Primitive->GetActorPositionForRenderer();
 
 			struct FPrimitiveUpdateParams
 			{
