Fix robomerge compile error

diff --git a/Engine/Plugins/Experimental/GPULightmass/Source/GPULightmass/Private/LightmapRenderer.cpp b/Engine/Plugins/Experimental/GPULightmass/Source/GPULightmass/Private/LightmapRenderer.cpp
--- a/Engine/Plugins/Experimental/GPULightmass/Source/GPULightmass/Private/LightmapRenderer.cpp
+++ b/Engine/Plugins/Experimental/GPULightmass/Source/GPULightmass/Private/LightmapRenderer.cpp
@@ -377,7 +377,6 @@ void FCachedRayTracingSceneData::SetupViewUniformBufferFromSceneRenderState(FSce
 					0, /* Custom Data Count */
 					0.0f, /* Random ID */
 					FRenderTransform::Identity,
-					PrimitiveUniformShaderParameters.LocalToRelativeWorld,
 					PrimitiveUniformShaderParameters.LocalToRelativeWorld
 				);
 
