Fix robomerge compile error

diff --git a/Engine/Plugins/Experimental/GPULightmass/Source/GPULightmass/Private/LightmapRenderer.cpp b/Engine/Plugins/Experimental/GPULightmass/Source/GPULightmass/Private/LightmapRenderer.cpp
--- a/Engine/Plugins/Experimental/GPULightmass/Source/GPULightmass/Private/LightmapRenderer.cpp
+++ b/Engine/Plugins/Experimental/GPULightmass/Source/GPULightmass/Private/LightmapRenderer.cpp
@@ -459,7 +459,6 @@ void FCachedRayTracingSceneData::SetupViewUniformBufferFromSceneRenderState(FSce
 					0, /* Custom Data Count */
 					0.0f, /* Random ID */
 					Instance.LocalToPrimitive,
-					InstanceGroup.LocalToWorld,
 					InstanceGroup.LocalToWorld
 				);
 
