Fixing non-unity win64 all modules compile errors

diff --git a/Engine/Plugins/Experimental/GPULightmass/Source/GPULightmass/Private/LightmapEncoding.cpp b/Engine/Plugins/Experimental/GPULightmass/Source/GPULightmass/Private/LightmapEncoding.cpp
--- a/Engine/Plugins/Experimental/GPULightmass/Source/GPULightmass/Private/LightmapEncoding.cpp
+++ b/Engine/Plugins/Experimental/GPULightmass/Source/GPULightmass/Private/LightmapEncoding.cpp
@@ -1,6 +1,7 @@
 // Copyright Epic Games, Inc. All Rights Reserved.
 
 #include "LightmapEncoding.h"
+#include "Async/ParallelFor.h"
 #include "GPULightmassCommon.h"
 #include "GPULightmassModule.h"
 
diff --git a/Engine/Plugins/Experimental/GPULightmass/Source/GPULightmass/Private/Scene/Scene.cpp b/Engine/Plugins/Experimental/GPULightmass/Source/GPULightmass/Private/Scene/Scene.cpp
--- a/Engine/Plugins/Experimental/GPULightmass/Source/GPULightmass/Private/Scene/Scene.cpp
+++ b/Engine/Plugins/Experimental/GPULightmass/Source/GPULightmass/Private/Scene/Scene.cpp
@@ -1,6 +1,7 @@
 // Copyright Epic Games, Inc. All Rights Reserved.
 
 #include "Scene.h"
+#include "Async/ParallelFor.h"
 #include "GPULightmass.h"
 #include "ComponentRecreateRenderStateContext.h"
 #include "Async/Async.h"
@@ -9,6 +10,7 @@
 #include "Components/StaticMeshComponent.h"
 #include "GameFramework/WorldSettings.h"
 #include "Framework/Notifications/NotificationManager.h"
+#include "StaticMeshComponentLODInfo.h"
 #include "Widgets/Notifications/SNotificationList.h"
 #include "Misc/ScopedSlowTask.h"
 #include "LightmapPreviewVirtualTexture.h"
diff --git a/Engine/Plugins/Runtime/nDisplay/Source/DisplayCluster/Private/Render/Synchronization/Windows/DisplayClusterRenderSyncHelperWindows.cpp b/Engine/Plugins/Runtime/nDisplay/Source/DisplayCluster/Private/Render/Synchronization/Windows/DisplayClusterRenderSyncHelperWindows.cpp
--- a/Engine/Plugins/Runtime/nDisplay/Source/DisplayCluster/Private/Render/Synchronization/Windows/DisplayClusterRenderSyncHelperWindows.cpp
+++ b/Engine/Plugins/Runtime/nDisplay/Source/DisplayCluster/Private/Render/Synchronization/Windows/DisplayClusterRenderSyncHelperWindows.cpp
@@ -6,6 +6,7 @@
 
 #include "Engine/Engine.h"
 #include "Engine/GameViewportClient.h"
+#include "RHI.h"
 #include "DynamicRHI.h"
 
 #pragma warning(push)
diff --git a/Engine/Plugins/Runtime/nDisplay/Source/DisplayCluster/Private/Render/Synchronization/Windows/DisplayClusterRenderSyncPolicyNvidiaWindows.cpp b/Engine/Plugins/Runtime/nDisplay/Source/DisplayCluster/Private/Render/Synchronization/Windows/DisplayClusterRenderSyncPolicyNvidiaWindows.cpp
--- a/Engine/Plugins/Runtime/nDisplay/Source/DisplayCluster/Private/Render/Synchronization/Windows/DisplayClusterRenderSyncPolicyNvidiaWindows.cpp
+++ b/Engine/Plugins/Runtime/nDisplay/Source/DisplayCluster/Private/Render/Synchronization/Windows/DisplayClusterRenderSyncPolicyNvidiaWindows.cpp
@@ -12,6 +12,7 @@
 
 #include "ID3D11DynamicRHI.h"
 
+#include "UnrealClient.h"
 #include "Windows/AllowWindowsPlatformTypes.h"
 THIRD_PARTY_INCLUDES_START
 	#include <nvapi.h>
diff --git a/Engine/Plugins/VirtualProduction/TextureShare/Source/TextureShare/Private/Module/TextureShareAPI.h b/Engine/Plugins/VirtualProduction/TextureShare/Source/TextureShare/Private/Module/TextureShareAPI.h
--- a/Engine/Plugins/VirtualProduction/TextureShare/Source/TextureShare/Private/Module/TextureShareAPI.h
+++ b/Engine/Plugins/VirtualProduction/TextureShare/Source/TextureShare/Private/Module/TextureShareAPI.h
@@ -6,6 +6,8 @@
 #include "Object/TextureShareObject.h"
 #include "Object/TextureShareObjectProxy.h"
 
+class SWindow;
+
 /**
  * TextureShare API impl
  */
