Fix non unity/pch compile error

diff --git a/Engine/Plugins/Experimental/GPULightmass/Source/GPULightmass/Private/Scene/Lights.cpp b/Engine/Plugins/Experimental/GPULightmass/Source/GPULightmass/Private/Scene/Lights.cpp
--- a/Engine/Plugins/Experimental/GPULightmass/Source/GPULightmass/Private/Scene/Lights.cpp
+++ b/Engine/Plugins/Experimental/GPULightmass/Source/GPULightmass/Private/Scene/Lights.cpp
@@ -2,6 +2,7 @@
 
 #include "Lights.h"
 
+#include "Engine/TextureLightProfile.h"
 #include "LightmapRayTracing.h"
 #include "Scene.h"
 #include "RenderGraphBuilder.h"
diff --git a/Engine/Source/Runtime/Renderer/Private/RendererScene.cpp b/Engine/Source/Runtime/Renderer/Private/RendererScene.cpp
--- a/Engine/Source/Runtime/Renderer/Private/RendererScene.cpp
+++ b/Engine/Source/Runtime/Renderer/Private/RendererScene.cpp
@@ -5,6 +5,7 @@
 =============================================================================*/
 
 #include "CoreMinimal.h"
+#include "Engine/TextureLightProfile.h"
 #include "HAL/ThreadSafeCounter.h"
 #include "HAL/PlatformFileManager.h"
 #include "Stats/Stats.h"
