Fix another CIS compilation error.

diff --git a/Engine/Plugins/Experimental/GPULightmass/Source/GPULightmass/Private/Scene/Lights.h b/Engine/Plugins/Experimental/GPULightmass/Source/GPULightmass/Private/Scene/Lights.h
--- a/Engine/Plugins/Experimental/GPULightmass/Source/GPULightmass/Private/Scene/Lights.h
+++ b/Engine/Plugins/Experimental/GPULightmass/Source/GPULightmass/Private/Scene/Lights.h
@@ -95,7 +95,7 @@ struct FLocalLightRenderState
 	int ShadowMapChannel = INDEX_NONE;
 	uint32 IESAtlasId = INDEX_NONE;
 	float IESAtlasIndex = INDEX_NONE;
-	UTexture* IESTexture = nullptr;
+	UTextureLightProfile* IESTexture = nullptr;
 	
 	TSharedPtr<FLightComponentMapBuildData, ESPMode::ThreadSafe> LightComponentMapBuildData;
 
