GPULM: Remove RT payload from compute shader that did not require it and was triggering an assertion

diff --git a/Engine/Plugins/Experimental/GPULightmass/Source/GPULightmass/Private/LightmapRayTracing.h b/Engine/Plugins/Experimental/GPULightmass/Source/GPULightmass/Private/LightmapRayTracing.h
--- a/Engine/Plugins/Experimental/GPULightmass/Source/GPULightmass/Private/LightmapRayTracing.h
+++ b/Engine/Plugins/Experimental/GPULightmass/Source/GPULightmass/Private/LightmapRayTracing.h
@@ -238,11 +238,6 @@ class FFirstBounceRayGuidingCDFBuildCS : public FGlobalShader
 		OutEnvironment.CompilerFlags.Add(CFLAG_WarningsAsErrors);
 	}
 
-	static ERayTracingPayloadType GetRayTracingPayloadType(const int32 PermutationId)
-	{
-		return ERayTracingPayloadType::Minimal | ERayTracingPayloadType::PathTracingMaterial;
-	}
-
 	BEGIN_SHADER_PARAMETER_STRUCT(FParameters, )
 		SHADER_PARAMETER(int, RayGuidingEndPassIndex)
 		SHADER_PARAMETER_SRV(StructuredBuffer<FGPUTileDescription>, BatchedTiles)
