GPULM: Fix remaining RT payload flags missed in previous refactor

diff --git a/Engine/Plugins/Experimental/GPULightmass/Source/GPULightmass/Private/LightmapRayTracing.h b/Engine/Plugins/Experimental/GPULightmass/Source/GPULightmass/Private/LightmapRayTracing.h
--- a/Engine/Plugins/Experimental/GPULightmass/Source/GPULightmass/Private/LightmapRayTracing.h
+++ b/Engine/Plugins/Experimental/GPULightmass/Source/GPULightmass/Private/LightmapRayTracing.h
@@ -267,6 +267,11 @@ class FStaticShadowDepthMapTracingRGS : public FGlobalShader
 		OutEnvironment.CompilerFlags.Add(CFLAG_WarningsAsErrors);
 	}
 
+	static ERayTracingPayloadType GetRayTracingPayloadType(const int32 PermutationId)
+	{
+		return ERayTracingPayloadType::Minimal | ERayTracingPayloadType::PathTracingMaterial;
+	}
+
 	BEGIN_SHADER_PARAMETER_STRUCT(FParameters, )
 		SHADER_PARAMETER_STRUCT_REF(FViewUniformShaderParameters, ViewUniformBuffer)
 		SHADER_PARAMETER(int32, StaticShadowDepthMapSuperSampleFactor)
