GPULM 5.1.1 Fixes

* Check meshes for lightmap UV version upgrade when starting a build
* Fix out-of-bounds of InstancedStaticMesh NumInstances calculation
* Reset LightmapTilePool's LRU FreeHeap when BumpRevision(). It was causing newly allocated tiles to be stepped on because FrameNumber is reset to 0.
* Fix WriteSortedBrickRequests hitting thread group limit when NumTotalBricks is large
* Fix CoordinateScale calculation when lightmap size is not aligned to 64 which causes it to be padded. Also hotfix SkyOcclusion access crash because the correspoding tile is not decompressed.
* Fix incorrect type of TilePositionsBufferSRV
* Fix ray bias issue when a light has a very close surface behind. Also fix wrong propagation of bIsValidSample when two sided materials are used.
* Surface lightmaps are now also sorted and processed in Morton order to improve irradiance caching efficiency
* When selecting lightmap tiles to process, do not mix different LODs together to improve spatial coherency
* Limit irradiance caching to a max lookup distance as cache entries that are very far are hard to build and have a very low reuse chance - clipmaps can potentially fix this

diff --git a/Engine/Plugins/Experimental/GPULightmass/Shaders/Private/BrickAllocationManagement.usf b/Engine/Plugins/Experimental/GPULightmass/Shaders/Private/BrickAllocationManagement.usf
--- a/Engine/Plugins/Experimental/GPULightmass/Shaders/Private/BrickAllocationManagement.usf
+++ b/Engine/Plugins/Experimental/GPULightmass/Shaders/Private/BrickAllocationManagement.usf
@@ -113,9 +113,12 @@ void SplatVolumeCS(
 	}
 }
 
-[numthreads(1, 1, 1)]
+int NumTotalBricks;
+
+[numthreads(64, 1, 1)]
 void WriteSortedBrickRequestsCS(uint3 RequestIndex : SV_DispatchThreadID)
 {
+	if (RequestIndex.x >= NumTotalBricks) return;
 	BrickRequests[RequestIndex.x] = BrickRequestsUnsorted[BrickRequestSortedIndices[RequestIndex.x]];
 	BrickRequestSortedIndicesInverse[BrickRequestSortedIndices[RequestIndex.x]] = RequestIndex.x;
 }
@@ -357,7 +360,6 @@ float3 ManualTrilinearFilter(RWTexture3D<UNORM float> Texture, float3 UVW)
 	return Value;
 }
 
-int NumTotalBricks;
 int BrickBatchOffset;
 RWTexture3D<uint> ValidityMask;
 
diff --git a/Engine/Plugins/Experimental/GPULightmass/Shaders/Private/LightmapPathTracing.usf b/Engine/Plugins/Experimental/GPULightmass/Shaders/Private/LightmapPathTracing.usf
--- a/Engine/Plugins/Experimental/GPULightmass/Shaders/Private/LightmapPathTracing.usf
+++ b/Engine/Plugins/Experimental/GPULightmass/Shaders/Private/LightmapPathTracing.usf
@@ -146,6 +146,7 @@ void ApplyRayBias(inout FRayDesc Ray, float HitT, float3 Normal, bool ShadowTerm
 	// for shadow rays, push out the point a bit more so we can reduce shadow terminator issues from poor tesselations
 	// other rays are less likely to arrive at grazing angles, so don't need the extra offset
 	Ray.Origin += max(ShadowTerminatorFix ? ShadowOffset : BaseOffset, RelativeOffset * 100) * Normal;
+	Ray.TMax -= length(max(ShadowTerminatorFix ? ShadowOffset : BaseOffset, RelativeOffset * 100) * Normal);
 #else
 	ApplyPositionBias(Ray.Origin, Ray.Direction, Normal, MaxNormalBias);
 #endif
@@ -427,7 +428,21 @@ void PathTracingKernel(
 		float DistanceFromTexelToHit = length(HitPosition - TexelPosition);
 		bool bRejectedByCornerRejection = RenderPassIndex >= GetIrradianceCachingQuality() && DistanceFromTexelToHit < IrradianceCachingParameters.Spacing * IrradianceCachingParameters.CornerRejection;
 		
-		if (NearestCacheEntryID == -1 && Bounce >= 1 && !bRejectedByCornerRejection)
+		// Compute an overestimated upper bound for the number of samples we need for this texel by assuming
+		// all the rays from all IC building passes are sent from this single texel
+		float CellArea = IrradianceCachingParameters.Spacing * IrradianceCachingParameters.Spacing;
+		float TotalSamplesOverAllICPasses = GPreviewLightmapPhysicalTileSize * GPreviewLightmapPhysicalTileSize * GetIrradianceCachingQuality();
+		float DistanceFraction = CellArea / (2 * PI * DistanceFromTexelToHit * DistanceFromTexelToHit);
+
+		// Cache entries too far away have a very low chance of being reused
+		// Don't bother placing or looking them up
+		bool bRejectedByMaxLookupDistace = false;
+		if (TotalSamplesOverAllICPasses * DistanceFraction < 1)
+		{
+			bRejectedByMaxLookupDistace = true;
+		}
+		
+		if (NearestCacheEntryID == -1 && Bounce >= 1 && !bRejectedByCornerRejection && !bRejectedByMaxLookupDistace)
 		{
 			bool bIrradianceQuerySuccessful = false;
 			bool bGeometryQuerySuccessful = false;
@@ -442,12 +457,6 @@ void PathTracingKernel(
 
 				uint4 RecordIrradianceAndSampleCount = IrradianceCachingParameters.IrradianceCacheRecords[RecordIndex];
 				uint BackfaceHitsCount = IrradianceCachingParameters.IrradianceCacheRecordBackfaceHits[RecordIndex].a;
-
-				// Compute an overestimated upper bound for the number of samples we need for this texel by assuming
-				// all the rays from all IC building passes are sent from this single texel
-				float CellArea = IrradianceCachingParameters.Spacing * IrradianceCachingParameters.Spacing;
-				float TotalSamplesOverAllICPasses = GPreviewLightmapPhysicalTileSize * GPreviewLightmapPhysicalTileSize * GetIrradianceCachingQuality();
-				float DistanceFraction = CellArea / (2 * PI * DistanceFromTexelToHit * DistanceFromTexelToHit);
 			
 				// We accept a cache entry, if:
 				// 1) The entry has been built by other lightmap tiles and contains more samples than we will ever need.
@@ -880,10 +889,10 @@ void LightmapPathTracingMainRG()
 
 	float4 EncodedGBufferWorldPosition = GBufferWorldPosition[TexelIndexInScratch];
 
-	bool bIsValidSample = true;
+	bool bGBufferSampleValid = true;
 	if (EncodedGBufferWorldPosition.w == 0.0f)
 	{
-		bIsValidSample = false;
+		bGBufferSampleValid = false;
 	}
 
 	float3 TranslatedWorldPosition = EncodedGBufferWorldPosition.xyz / WorldPositionScalar + LWCHackToFloat(PrimaryView.PreViewTranslation); // RT_LWC_TODO
@@ -912,15 +921,18 @@ void LightmapPathTracingMainRG()
 	
 	if (NumTotalSamples <= 0 || (NumTotalSamples > 0 && BatchedTiles[TileIndex].RenderPassIndex < NumTotalSamples - 1))
 	{
-		if (bIsValidSample)
+		bool bAnyHemispherePathSampleValid = false;
+	
+		if (bGBufferSampleValid)
 		{			
 			// Needs a seed that is only related to position in virtual space to avoid seams due to per-tile calculation
 			uint2 VirtualTextureSpacePosition = BatchedTiles[TileIndex].VirtualTilePosition + LaunchIndex - uint2(2, 2);
 			uint Seed = VirtualTextureSpacePosition.y * BatchedTiles[TileIndex].LightmapSize.x + VirtualTextureSpacePosition.x;
 					
 			float4 PrimaryRandSample = float4(-1, -1, -1, -1);
-
+			
 			// Evaluate both hemispheres for two sided materials and add the lighting up naively
+			UNROLL_N(2)
 			for (int HemisphereIndex = 0; HemisphereIndex < (bMaterialTwoSided ? 2 : 1); HemisphereIndex++)
 			{
 				RandomSequence RandSequence;
@@ -940,32 +952,31 @@ void LightmapPathTracingMainRG()
 				float3 SkyLightBentNormal = 0;
 
 				float3 EffectiveShadingNormal = HemisphereIndex == 0 ? ShadingNormal : -ShadingNormal;
+
+				bool bPathSampleValid = true;
 				
-				if (bIsValidSample)
-				{
-					PathTracingKernel(
-						BatchedTiles[TileIndex].RenderPassIndex,
-						TranslatedWorldPosition,
-						EffectiveShadingNormal,
-						RandSequence,
-						SampleIndex,
-						bIsValidSample,
-						RadianceValue,
-						RadianceDirection,
-						DirectLightingIrradiance,
-						DirectLightingSH,
-						LuminanceForFirstBounceRayGuiding,
-						SkyLightBentNormal,
-						PrimaryRandSample
-					);
-				}
+				PathTracingKernel(
+					BatchedTiles[TileIndex].RenderPassIndex,
+					TranslatedWorldPosition,
+					EffectiveShadingNormal,
+					RandSequence,
+					SampleIndex,
+					bPathSampleValid,
+					RadianceValue,
+					RadianceDirection,
+					DirectLightingIrradiance,
+					DirectLightingSH,
+					LuminanceForFirstBounceRayGuiding,
+					SkyLightBentNormal,
+					PrimaryRandSample
+				);
 				
 				if (any(isnan(RadianceValue)) || any(RadianceValue < 0) || any(isinf(RadianceValue)))
 				{
-					bIsValidSample = false;
+					bPathSampleValid = false;
 				}
 				
-				if (bIsValidSample)
+				if (bPathSampleValid)
 				{
 				#if USE_FIRST_BOUNCE_RAY_GUIDING
 					int MinRenderPassIndex = 0;
@@ -1013,6 +1024,8 @@ void LightmapPathTracingMainRG()
 
 					SHCorrectionAndStationarySkyLightBentNormal[TexelIndexInPool].yzw += SkyLightBentNormal;
 				}
+
+				bAnyHemispherePathSampleValid = bAnyHemispherePathSampleValid || bPathSampleValid;
 			}
 		}
 
@@ -1032,7 +1045,7 @@ void LightmapPathTracingMainRG()
 		}
 #endif
 		
-		if (bIsValidSample)
+		if (bAnyHemispherePathSampleValid)
 			SampleCount++;
 		
 		IrradianceAndSampleCount[TexelIndexInPool].a = asfloat(SampleCount);
@@ -1271,6 +1284,7 @@ void VolumetricLightmapPathTracingMainRG()
 
 	int NumValidHalfSamples = 0;
 	
+	UNROLL_N(2)
 	for (int HalfSampleIndex = 0; HalfSampleIndex < 2; HalfSampleIndex++)
 	{
 		float3 ShadingNormal = float3(0, 0, (HalfSampleIndex == 0) ? 1 : -1);
@@ -1429,21 +1443,22 @@ void StationaryLightShadowTracingMainRG()
 
 	int ChannelIndex = ChannelIndexArray[TileIndex];
 
-	bool bIsValidSample = true;
+	bool bGBufferSampleValid = true;
 
 	float4 EncodedGBufferWorldPosition = GBufferWorldPosition[TexelIndexInScratch];
 	
 	if (EncodedGBufferWorldPosition.w == 0.0f)
 	{
-		bIsValidSample = false;
+		bGBufferSampleValid = false;
 	}
 
 	float3 TranslatedWorldPosition = EncodedGBufferWorldPosition.xyz / WorldPositionScalar + LWCHackToFloat(PrimaryView.PreViewTranslation); // RT_LWC_TODO
 	
-	if (bIsValidSample)
+	if (bGBufferSampleValid)
 	{
 		float3 WorldNormal = GBufferWorldNormal[TexelIndexInScratch].xyz;
 		float3 ShadingNormal = GBufferShadingNormal[TexelIndexInScratch].xyz;
+		bool bMaterialTwoSided = GBufferShadingNormal[TexelIndexInScratch].w;
 		
 		// Needs a seed that is only related to position in virtual space to avoid seams due to per-tile calculation
 		uint2 VirtualTextureSpacePosition = BatchedTiles[TileIndex].VirtualTilePosition + LaunchIndex - uint2(2, 2);
@@ -1452,27 +1467,36 @@ void StationaryLightShadowTracingMainRG()
 		int LightType = LightTypeArray[TileIndex];
 		FLightShaderParameters LightParameters;
 		LightShaderParametersArray[TileIndex].FillLightShaderParameters(LightParameters);
+
+		bool bAnyHemisphereRaySampleValid = false;
+		uint Visibility = 0;
 		
+		UNROLL_N(2)
+		for (int HemisphereIndex = 0; HemisphereIndex < (bMaterialTwoSided ? 2 : 1); HemisphereIndex++)
 		{
 			float2 RandSample = float2(
 				Halton(StrongIntegerHash(Seed) + LightSampleIndexArray[TileIndex], 5), 
 				Halton(StrongIntegerHash(Seed) + LightSampleIndexArray[TileIndex], 7)
 			);
+
+			float3 EffectiveWorldNormal = HemisphereIndex == 0 ? WorldNormal : -WorldNormal;
+
+			bool bRaySampleValid = true;
 			
 			FRayDesc Ray;
-			bIsValidSample &= GenerateOcclusionRay(
+			bRaySampleValid &= GenerateOcclusionRay(
 				LightType,
 				LightParameters,
-				TranslatedWorldPosition, WorldNormal,
+				TranslatedWorldPosition, EffectiveWorldNormal,
 				RandSample,
 				/* out */ Ray.Origin,
 				/* out */ Ray.Direction,
 				/* out */ Ray.TMin,
 				/* out */ Ray.TMax);
 				
-			ApplyRayBias(Ray, 1.0f, WorldNormal);
-				
-			if (bIsValidSample)
+			ApplyRayBias(Ray, LightType == LIGHT_TYPE_DIRECTIONAL ? 1.0f : max(Ray.TMax - Ray.TMin, 0.0f), EffectiveWorldNormal);
+			
+			if (bRaySampleValid)
 			{
 				uint RayFlags = RAY_FLAG_SKIP_CLOSEST_HIT_SHADER | RAY_FLAG_ACCEPT_FIRST_HIT_AND_END_SEARCH;
 
@@ -1481,13 +1505,18 @@ void StationaryLightShadowTracingMainRG()
 					RayFlags,
 					Ray);
 				
-				uint Visibility = asuint(ShadowMask[TexelIndexInPool][ChannelIndex]);
-				uint SampleCount = asuint(ShadowMaskSampleCount[TexelIndexInPool][ChannelIndex]);
-				Visibility += MinimalPayload.IsMiss() ? 1 : 0;
-				SampleCount++;
-				ShadowMask[TexelIndexInPool][ChannelIndex] = asfloat(Visibility);
-				ShadowMaskSampleCount[TexelIndexInPool][ChannelIndex] = asfloat(SampleCount);
+				Visibility = max(Visibility, MinimalPayload.IsMiss() ? 1 : 0);
 			}
+
+			bAnyHemisphereRaySampleValid = bAnyHemisphereRaySampleValid || bRaySampleValid;
+		}
+
+		if (bAnyHemisphereRaySampleValid)
+		{
+			ShadowMask[TexelIndexInPool][ChannelIndex] = asfloat(asuint(ShadowMask[TexelIndexInPool][ChannelIndex]) + Visibility);
+			uint SampleCount = asuint(ShadowMaskSampleCount[TexelIndexInPool][ChannelIndex]);
+			SampleCount++;
+			ShadowMaskSampleCount[TexelIndexInPool][ChannelIndex] = asfloat(SampleCount);
 		}
 	}
 }
diff --git a/Engine/Plugins/Experimental/GPULightmass/Source/GPULightmass/Private/GPULightmassSettings.cpp b/Engine/Plugins/Experimental/GPULightmass/Source/GPULightmass/Private/GPULightmassSettings.cpp
--- a/Engine/Plugins/Experimental/GPULightmass/Source/GPULightmass/Private/GPULightmassSettings.cpp
+++ b/Engine/Plugins/Experimental/GPULightmass/Source/GPULightmass/Private/GPULightmassSettings.cpp
@@ -10,6 +10,9 @@
 #include "Editor.h"
 #include "LevelEditorViewport.h"
 #include "AssetCompilingManager.h"
+#include "MeshUtilitiesCommon.h"
+#include "Logging/MessageLog.h"
+#include "Misc/UObjectToken.h"
 
 #define LOCTEXT_NAMESPACE "StaticLightingSystem"
 
@@ -136,6 +139,51 @@ void UGPULightmassSubsystem::Initialize(FSubsystemCollectionBase& Collection)
 	}
 }
 
+bool CheckStaticMeshesForLightmapUVVersionUpgrade(const UWorld* World)
+{
+	TSet<UStaticMesh*> StaticMeshesToUpgradeLightmapUVVersion;
+			
+	for (const UStaticMeshComponent* Component : TObjectRange<UStaticMeshComponent>(RF_ClassDefaultObject | RF_ArchetypeObject, true, EInternalObjectFlags::Garbage))
+	{
+		if (Component->GetWorld() == World && Component->GetStaticMesh() && Component->GetStaticMesh()->GetLightmapUVVersion() != (int32)ELightmapUVVersion::Latest)
+		{
+			StaticMeshesToUpgradeLightmapUVVersion.Add(Component->GetStaticMesh().Get());
+		}
+	}
+
+	if (!StaticMeshesToUpgradeLightmapUVVersion.IsEmpty())
+	{
+		const EAppReturnType::Type Response = FMessageDialog::Open(EAppMsgType::YesNo,
+			LOCTEXT("MeshLMUVNotLatestMessage", "Some meshes in your scene are not using the latest lightmap UV generation algorithm. Would you like to upgrade them? This can break existing lightmaps."));
+
+		if (Response == EAppReturnType::Yes)
+		{
+			if (GCurrentLevelEditingViewportClient)
+			{
+				GCurrentLevelEditingViewportClient->SetShowStats(true);
+			}
+					
+			UStaticMesh::BatchBuild(StaticMeshesToUpgradeLightmapUVVersion.Array());
+					
+			for (UStaticMesh* Mesh : StaticMeshesToUpgradeLightmapUVVersion)
+			{
+				FMessageLog("LightingResults").Info()
+					->AddToken(FUObjectToken::Create(Mesh))
+					->AddToken(FTextToken::Create(LOCTEXT("LightingResults_LMUVVersionUpdated", "lightmap UV version has been upgraded and need to be resaved.")));
+						
+				Mesh->MarkPackageDirty();
+				Mesh->OnMeshChanged.Broadcast();
+			}
+
+			FMessageLog("LightingResults").Open();
+					
+			return true;
+		}
+	}
+
+	return false;
+}
+
 void UGPULightmassSubsystem::Launch()
 {
 	UWorld* World = Cast<UWorld>(GetOuter());
@@ -148,6 +196,11 @@ void UGPULightmassSubsystem::Launch()
 	{
 		UGPULightmassSettings* SettingsCopy = DuplicateObject(GetSettings(), GetTransientPackage(), MakeUniqueObjectName(GetTransientPackage(), UGPULightmassSettings::StaticClass()));
 
+		if (CheckStaticMeshesForLightmapUVVersionUpgrade(World))
+		{
+			return;
+		}
+		
 		FScopedSlowTask SlowTask(1);
 		SlowTask.MakeDialog();
 		SlowTask.EnterProgressFrame(1, LOCTEXT("StartingStaticLightingSystem", "Starting static lighting system"));
diff --git a/Engine/Plugins/Experimental/GPULightmass/Source/GPULightmass/Private/LightmapRenderer.cpp b/Engine/Plugins/Experimental/GPULightmass/Source/GPULightmass/Private/LightmapRenderer.cpp
--- a/Engine/Plugins/Experimental/GPULightmass/Source/GPULightmass/Private/LightmapRenderer.cpp
+++ b/Engine/Plugins/Experimental/GPULightmass/Source/GPULightmass/Private/LightmapRenderer.cpp
@@ -202,7 +202,7 @@ struct FGPUBatchedTileRequests
 			CreateInfo.GPUMask = FRHIGPUMask::FromIndex(GPUIndex);
 			CreateInfo.ResourceArray = &BatchedTilesDesc;
 
-			BatchedTilesBuffer = RHICreateStructuredBuffer(sizeof(FGPUTileDescription), BatchedTilesDesc.GetResourceDataSize(), BUF_Dynamic | BUF_ShaderResource, CreateInfo);
+			BatchedTilesBuffer = RHICreateStructuredBuffer(sizeof(FGPUTileDescription), BatchedTilesDesc.GetResourceDataSize(), BUF_Static | BUF_ShaderResource, CreateInfo);
 			BatchedTilesSRV = RHICreateShaderResourceView(BatchedTilesBuffer);
 		}
 	}
@@ -1509,6 +1509,35 @@ void FLightmapRenderer::RenderMeshBatchesIntoGBuffer(
 	GPrimitiveIdVertexBufferPool.DiscardAll();
 }
 
+void ClearScratchTilePoolForMultipleTiles(
+	FRDGBuilder& GraphBuilder,
+	const TResourceArray<FIntPoint>& TilePositionsToClear,
+	const TStaticArray<FRDGTextureUAVRef, 3>& ScratchTilePoolLayerUAVs,
+	const FGlobalShaderMap* GlobalShaderMap)
+{
+	for (int ScratchLayerIndex = 0; ScratchLayerIndex < ScratchTilePoolLayerUAVs.Num(); ScratchLayerIndex++)
+	{
+		RDG_GPU_MASK_SCOPE(GraphBuilder, FRHIGPUMask::All());
+
+		FRDGBufferDesc TilePositionsBufferDesc = FRDGBufferDesc::CreateBufferDesc(TilePositionsToClear.GetTypeSize(), TilePositionsToClear.Num());
+		const FRDGBufferRef TilePositionsBuffer = CreateVertexBuffer(GraphBuilder, TEXT("TilePositionsBufferForClear"), TilePositionsBufferDesc, TilePositionsToClear.GetData(), TilePositionsToClear.GetResourceDataSize());
+		const FRDGBufferSRVRef TilePositionsBufferSRV = GraphBuilder.CreateSRV(FRDGBufferSRVDesc(TilePositionsBuffer, PF_R32G32_UINT));
+		FMultiTileClearCS::FParameters* Parameters = GraphBuilder.AllocParameters<FMultiTileClearCS::FParameters>();
+		Parameters->NumTiles = TilePositionsToClear.Num();
+		Parameters->TileSize = GPreviewLightmapPhysicalTileSize;
+		Parameters->TilePositions = TilePositionsBufferSRV;
+		Parameters->TilePool = ScratchTilePoolLayerUAVs[ScratchLayerIndex];
+
+		TShaderMapRef<FMultiTileClearCS> ComputeShader(GlobalShaderMap);
+		FComputeShaderUtils::AddPass(
+			GraphBuilder,
+			RDG_EVENT_NAME("MultiTileClear"),
+			ComputeShader,
+			Parameters,
+			FComputeShaderUtils::GetGroupCount(FIntPoint(GPreviewLightmapPhysicalTileSize * TilePositionsToClear.Num(), GPreviewLightmapPhysicalTileSize), FComputeShaderUtils::kGolden2DGroupSize));
+	}
+}
+	
 void FLightmapRenderer::Finalize(FRDGBuilder& GraphBuilder)
 {
 	TRACE_CPUPROFILER_EVENT_SCOPE(FLightmapRenderer::Finalize);
@@ -2175,33 +2204,13 @@ void FLightmapRenderer::Finalize(FRDGBuilder& GraphBuilder)
 						PendingGIRenderPassIndices.Add(Tile.RenderState->RetrieveTileState(Tile.VirtualCoordinates).RenderPassIndex);
 					}
 
-					for (int ScratchLayerIndex = 0; ScratchLayerIndex < 3; ScratchLayerIndex++)
+					TResourceArray<FIntPoint> TilePositionsToClear;
+					for (auto& Tile : PendingGITileRequests)
 					{
-						RDG_GPU_MASK_SCOPE(GraphBuilder, FRHIGPUMask::All());
-						
-						TResourceArray<FIntPoint> TilePositionsToClear;
-						for (auto& Tile : PendingGITileRequests)
-						{
-							TilePositionsToClear.Add(ScratchTilePoolGPU->GetPositionFromLinearAddress(Tile.TileAddressInScratch));
-						}
-
-						FRDGBufferDesc TilePositionsBufferDesc = FRDGBufferDesc::CreateBufferDesc(TilePositionsToClear.GetTypeSize(), TilePositionsToClear.Num());
-						FRDGBufferRef TilePositionsBuffer = CreateVertexBuffer(GraphBuilder, TEXT("TilePositionsBufferForClear"), TilePositionsBufferDesc, TilePositionsToClear.GetData(), TilePositionsToClear.GetResourceDataSize());
-						FRDGBufferSRVRef TilePositionsBufferSRV = GraphBuilder.CreateSRV(FRDGBufferSRVDesc(TilePositionsBuffer, PF_R32G32_UINT));
-						FMultiTileClearCS::FParameters* Parameters = GraphBuilder.AllocParameters<FMultiTileClearCS::FParameters>();
-						Parameters->NumTiles = TilePositionsToClear.Num();
-						Parameters->TileSize = GPreviewLightmapPhysicalTileSize;
-						Parameters->TilePositions = TilePositionsBufferSRV;
-						Parameters->TilePool = ScratchTilePoolLayerUAVs[ScratchLayerIndex];
-
-						TShaderMapRef<FMultiTileClearCS> ComputeShader(GlobalShaderMap);
-						FComputeShaderUtils::AddPass(
-							GraphBuilder,
-							RDG_EVENT_NAME("MultiTileClear"),
-							ComputeShader,
-							Parameters,
-							FComputeShaderUtils::GetGroupCount(FIntPoint(GPreviewLightmapPhysicalTileSize * TilePositionsToClear.Num(), GPreviewLightmapPhysicalTileSize), FComputeShaderUtils::kGolden2DGroupSize));
+						TilePositionsToClear.Add(ScratchTilePoolGPU->GetPositionFromLinearAddress(Tile.TileAddressInScratch));
 					}
+					
+					ClearScratchTilePoolForMultipleTiles(GraphBuilder, TilePositionsToClear, ScratchTilePoolLayerUAVs, GlobalShaderMap);
 
 					{
 						for (uint32 GPUIndex = 0; GPUIndex < GNumExplicitGPUsForRendering; GPUIndex++)
@@ -2466,33 +2475,14 @@ void FLightmapRenderer::Finalize(FRDGBuilder& GraphBuilder)
 
 					if (PendingShadowTileRequests.Num() == 0) continue;
 
-					for (int ScratchLayerIndex = 0; ScratchLayerIndex < 3; ScratchLayerIndex++)
+					TResourceArray<FIntPoint> TilePositionsToClear;
+					for (auto& Tile : PendingShadowTileRequests)
 					{
-						TResourceArray<FIntPoint> TilePositionsToClear;
-						for (auto& Tile : PendingShadowTileRequests)
-						{
-							TilePositionsToClear.Add(ScratchTilePoolGPU->GetPositionFromLinearAddress(Tile.TileAddressInScratch));
-						}
-
-						FRDGBufferDesc TilePositionsBufferDesc = FRDGBufferDesc::CreateBufferDesc(TilePositionsToClear.GetTypeSize(), TilePositionsToClear.Num());
-						FRDGBufferRef TilePositionsBuffer = CreateVertexBuffer(GraphBuilder, TEXT("TilePositionsBufferForClear"), TilePositionsBufferDesc, TilePositionsToClear.GetData(), TilePositionsToClear.GetResourceDataSize());
-						FRDGBufferSRVRef TilePositionsBufferSRV = GraphBuilder.CreateSRV(FRDGBufferSRVDesc(TilePositionsBuffer, PF_R32_UINT));
-
-						FMultiTileClearCS::FParameters* Parameters = GraphBuilder.AllocParameters<FMultiTileClearCS::FParameters>();
-						Parameters->NumTiles = TilePositionsToClear.Num();
-						Parameters->TileSize = GPreviewLightmapPhysicalTileSize;
-						Parameters->TilePositions = TilePositionsBufferSRV;
-						Parameters->TilePool = ScratchTilePoolLayerUAVs[ScratchLayerIndex];
-
-						TShaderMapRef<FMultiTileClearCS> ComputeShader(GlobalShaderMap);
-						FComputeShaderUtils::AddPass(
-							GraphBuilder,
-							RDG_EVENT_NAME("MultiTileClear"),
-							ComputeShader,
-							Parameters,
-							FComputeShaderUtils::GetGroupCount(FIntPoint(GPreviewLightmapPhysicalTileSize * TilePositionsToClear.Num(), GPreviewLightmapPhysicalTileSize), FComputeShaderUtils::kGolden2DGroupSize));
+						TilePositionsToClear.Add(ScratchTilePoolGPU->GetPositionFromLinearAddress(Tile.TileAddressInScratch));
 					}
 
+					ClearScratchTilePoolForMultipleTiles(GraphBuilder, TilePositionsToClear, ScratchTilePoolLayerUAVs, GlobalShaderMap);
+
 					FRDGTextureRef GBufferWorldPosition = GraphBuilder.RegisterExternalTexture(ScratchTilePoolGPU->PooledRenderTargets[0], TEXT("GBufferWorldPosition"));
 					FRDGTextureRef GBufferWorldNormal = GraphBuilder.RegisterExternalTexture(ScratchTilePoolGPU->PooledRenderTargets[1], TEXT("GBufferWorldNormal"));
 					FRDGTextureRef GBufferShadingNormal = GraphBuilder.RegisterExternalTexture(ScratchTilePoolGPU->PooledRenderTargets[2], TEXT("GBufferShadingNormal"));
@@ -3139,6 +3129,51 @@ void FLightmapTileDenoiseAsyncTask::DoThreadedWork()
 	FPlatformAtomics::AtomicStore(&TextureData->bDenoisingFinished, 1);
 }
 
+bool CompareMostSignificantBit(uint32 A, uint32 B)
+{
+	return A < B && A < (A ^ B);
+}
+
+bool MortonCompare(FUintVector A, FUintVector B)
+{
+	int32 MostSignificantDim = 0;
+	for (int32 Dim = 1; Dim < FUintVector::Num(); Dim++)
+	{
+		if (CompareMostSignificantBit(A[MostSignificantDim] ^ B[MostSignificantDim], A[Dim] ^ B[Dim]))
+		{
+			MostSignificantDim = Dim;
+		}
+	}
+
+	return A[MostSignificantDim] < B[MostSignificantDim];
+}
+
+void FSceneRenderState::BuildMortonSortedLightmapRefList()
+{
+	FBox SceneBounds;
+
+	for (FLightmapRenderState& Lightmap : LightmapRenderStates.Elements)
+	{
+		SceneBounds += Lightmap.GeometryInstanceRef.GetOrigin();
+	}
+
+	for (FLightmapRenderState& Lightmap : LightmapRenderStates.Elements)
+	{
+		MortonSortedLightmapRefList.Add(FLightmapRenderStateRef(Lightmap, LightmapRenderStates));
+	}
+
+	if (SceneBounds.GetSize().GetMax() > UE_DOUBLE_SMALL_NUMBER)
+	{
+		MortonSortedLightmapRefList.Sort(
+			[SceneBounds](const FLightmapRenderStateRef& A, const FLightmapRenderStateRef& B)
+			{
+				const FUintVector QuantizedPosA = FUintVector((A->GeometryInstanceRef.GetOrigin() - SceneBounds.Min) / SceneBounds.GetSize().GetMax() * UINT_MAX);
+				const FUintVector QuantizedPosB = FUintVector((B->GeometryInstanceRef.GetOrigin() - SceneBounds.Min) / SceneBounds.GetSize().GetMax() * UINT_MAX);
+				return MortonCompare(QuantizedPosA, QuantizedPosB);
+			});
+	}
+}
+
 void FLightmapRenderer::BackgroundTick()
 {
 	{
@@ -3538,31 +3573,38 @@ void FLightmapRenderer::BackgroundTick()
 
 			TArray<FString> SelectedLightmapNames;
 
+			if (Scene->MortonSortedLightmapRefList.IsEmpty())
+			{
+				Scene->BuildMortonSortedLightmapRefList();
+			}
+
 			// We schedule VLM work to be after lightmaps (which forces LOD 0). Making LOD 0 last here reduces the chance of rebuilding cached scene
 			for (int32 LODIndex = MAX_STATIC_MESH_LODS - 1; LODIndex >= 0; LODIndex--)
 			{
-				for (FLightmapRenderState& Lightmap : Scene->LightmapRenderStates.Elements)
+				bool bLODIndexSelected = false;
+				
+				for (const FLightmapRenderStateRef& Lightmap : Scene->MortonSortedLightmapRefList)
 				{
-					if (Lightmap.GeometryInstanceRef.LODIndex != LODIndex)
+					if (Lightmap->GeometryInstanceRef.LODIndex != LODIndex)
 					{
 						continue;
 					}
 
 					bool bAnyTileSelected = false;
 
-					for (int32 Y = 0; Y < Lightmap.GetPaddedSizeInTiles().Y; Y++)
+					for (int32 Y = 0; Y < Lightmap->GetPaddedSizeInTiles().Y; Y++)
 					{
-						for (int32 X = 0; X < Lightmap.GetPaddedSizeInTiles().X; X++)
+						for (int32 X = 0; X < Lightmap->GetPaddedSizeInTiles().X; X++)
 						{
 							FTileVirtualCoordinates VirtualCoordinates(FIntPoint(X, Y), 0);
 
-							if (!Lightmap.DoesTileHaveValidCPUData(VirtualCoordinates, CurrentRevision) && Lightmap.RetrieveTileState(VirtualCoordinates).OngoingReadbackRevision != CurrentRevision)
+							if (!Lightmap->DoesTileHaveValidCPUData(VirtualCoordinates, CurrentRevision) && Lightmap->RetrieveTileState(VirtualCoordinates).OngoingReadbackRevision != CurrentRevision)
 							{
 								bAnyTileSelected = true;
 
 								FVTProduceTargetLayer TargetLayers[4];
 
-								Lightmap.LightmapPreviewVirtualTexture->ProducePageData(
+								Lightmap->LightmapPreviewVirtualTexture->ProducePageData(
 									RHICmdList,
 									Scene->FeatureLevel,
 									EVTProducePageFlags::None,
@@ -3590,7 +3632,8 @@ void FLightmapRenderer::BackgroundTick()
 
 					if (bAnyTileSelected)
 					{
-						SelectedLightmapNames.Add(Lightmap.Name);
+						bLODIndexSelected = true;
+						SelectedLightmapNames.Add(Lightmap->Name);
 					}
 
 					if (WorkGenerated >= WorkToGenerate)
@@ -3599,13 +3642,14 @@ void FLightmapRenderer::BackgroundTick()
 					}
 				}
 
-				if (WorkGenerated >= WorkToGenerate)
+				// Do not mix different LODs together
+				if (bLODIndexSelected || WorkGenerated >= WorkToGenerate)
 				{
 					break;
 				}
 			}
 
-			if (bIsViewportNonRealtime && FrameNumber % 100 == 0)
+			if (!SelectedLightmapNames.IsEmpty() && bIsViewportNonRealtime && FrameNumber % 100 == 0)
 			{
 				FString AllNames;
 				for (FString& Name : SelectedLightmapNames)
@@ -3725,6 +3769,8 @@ void FLightmapRenderer::BumpRevision()
 	RecordedTileRequests.Empty();
 
 	LightmapTilePoolGPU.UnmapAll();
+
+	Scene->MortonSortedLightmapRefList.Empty();
 }
 
 void FLightmapRenderer::DeduplicateRecordedTileRequests()
diff --git a/Engine/Plugins/Experimental/GPULightmass/Source/GPULightmass/Private/LightmapStorage.h b/Engine/Plugins/Experimental/GPULightmass/Source/GPULightmass/Private/LightmapStorage.h
--- a/Engine/Plugins/Experimental/GPULightmass/Source/GPULightmass/Private/LightmapStorage.h
+++ b/Engine/Plugins/Experimental/GPULightmass/Source/GPULightmass/Private/LightmapStorage.h
@@ -76,6 +76,7 @@ public:
 			FMath::DivideAndRoundUp(Size.X, GPreviewLightmapVirtualTileSize),
 			FMath::DivideAndRoundUp(Size.Y, GPreviewLightmapVirtualTileSize));
 	}
+	FIntPoint GetPaddedSize() const { return GetPaddedSizeInTiles() * GPreviewLightmapVirtualTileSize; }
 
 	FString Name;
 	FIntPoint Size;
diff --git a/Engine/Plugins/Experimental/GPULightmass/Source/GPULightmass/Private/LightmapTilePool.cpp b/Engine/Plugins/Experimental/GPULightmass/Source/GPULightmass/Private/LightmapTilePool.cpp
--- a/Engine/Plugins/Experimental/GPULightmass/Source/GPULightmass/Private/LightmapTilePool.cpp
+++ b/Engine/Plugins/Experimental/GPULightmass/Source/GPULightmass/Private/LightmapTilePool.cpp
@@ -100,6 +100,12 @@ void FLightmapTilePool::UnmapAll()
 {
 	TileIndexToVirtualTileMap.Empty();
 	VirtualTileToTileIndexMap.Empty();
+		
+	FreeHeap.Free();	
+	for (int32 TileIndex = 0; TileIndex < NumTotalTiles; TileIndex++)
+	{
+		FreeHeap.Add(0, TileIndex);
+	}
 }
 
 FLightmapTilePoolGPU::FLightmapTilePoolGPU(int32 NumLayers, FIntPoint InSizeInTiles, FIntPoint TileSize)
diff --git a/Engine/Plugins/Experimental/GPULightmass/Source/GPULightmass/Private/Scene/GeometryInterface.cpp b/Engine/Plugins/Experimental/GPULightmass/Source/GPULightmass/Private/Scene/GeometryInterface.cpp
--- a/Engine/Plugins/Experimental/GPULightmass/Source/GPULightmass/Private/Scene/GeometryInterface.cpp
+++ b/Engine/Plugins/Experimental/GPULightmass/Source/GPULightmass/Private/Scene/GeometryInterface.cpp
@@ -11,4 +11,9 @@ TArray<FMeshBatch> FGeometryInstanceRenderStateRef::GetMeshBatchesForGBufferRend
 	return Collection.GetMeshBatchesForGBufferRendering(*this, CoordsForCulling);
 }
 
+FVector FGeometryInstanceRenderStateRef::GetOrigin() const
+{
+	return Collection.Get(GetElementIdChecked()).LocalToWorld.GetOrigin();
+}
+
 }
\ No newline at end of file
diff --git a/Engine/Plugins/Experimental/GPULightmass/Source/GPULightmass/Private/Scene/GeometryInterface.h b/Engine/Plugins/Experimental/GPULightmass/Source/GPULightmass/Private/Scene/GeometryInterface.h
--- a/Engine/Plugins/Experimental/GPULightmass/Source/GPULightmass/Private/Scene/GeometryInterface.h
+++ b/Engine/Plugins/Experimental/GPULightmass/Source/GPULightmass/Private/Scene/GeometryInterface.h
@@ -59,6 +59,7 @@ public:
 	{}
 
 	TArray<FMeshBatch> GetMeshBatchesForGBufferRendering(FTileVirtualCoordinates CoordsForCulling = FTileVirtualCoordinates{});
+	FVector GetOrigin() const;
 
 	int32 LODIndex;
 private:
diff --git a/Engine/Plugins/Experimental/GPULightmass/Source/GPULightmass/Private/Scene/InstancedStaticMesh.cpp b/Engine/Plugins/Experimental/GPULightmass/Source/GPULightmass/Private/Scene/InstancedStaticMesh.cpp
--- a/Engine/Plugins/Experimental/GPULightmass/Source/GPULightmass/Private/Scene/InstancedStaticMesh.cpp
+++ b/Engine/Plugins/Experimental/GPULightmass/Source/GPULightmass/Private/Scene/InstancedStaticMesh.cpp
@@ -204,8 +204,10 @@ TArray<FMeshBatch> FInstanceGroupRenderState::GetMeshBatchesForGBufferRendering(
 							int32 MinInstanceIndex = Y * InstancesPerRow + MinInInstanceTile.X;
 							int32 MaxInstanceIndex = Y * InstancesPerRow + MaxInInstanceTile.X;
 
-							MeshBatchElement.NumInstances = FMath::Min(MaxInstanceIndex - MinInstanceIndex, (int32)InstancedRenderData->PerInstanceRenderData->InstanceBuffer.GetNumInstances());
-							MeshBatchElement.UserIndex = FMath::Min(MinInstanceIndex, (int32)InstancedRenderData->PerInstanceRenderData->InstanceBuffer.GetNumInstances());
+							MinInstanceIndex = FMath::Min(MinInstanceIndex, (int32)InstancedRenderData->PerInstanceRenderData->InstanceBuffer.GetNumInstances() - 1);
+							
+							MeshBatchElement.UserIndex = MinInstanceIndex;							
+							MeshBatchElement.NumInstances = FMath::Min(MaxInstanceIndex - MinInstanceIndex, (int32)InstancedRenderData->PerInstanceRenderData->InstanceBuffer.GetNumInstances() - MinInstanceIndex);
 
 							MeshBatches.Add(MeshBatch);
 						}
diff --git a/Engine/Plugins/Experimental/GPULightmass/Source/GPULightmass/Private/Scene/Scene.cpp b/Engine/Plugins/Experimental/GPULightmass/Source/GPULightmass/Private/Scene/Scene.cpp
--- a/Engine/Plugins/Experimental/GPULightmass/Source/GPULightmass/Private/Scene/Scene.cpp
+++ b/Engine/Plugins/Experimental/GPULightmass/Source/GPULightmass/Private/Scene/Scene.cpp
@@ -967,7 +967,7 @@ void FScene::AddGeometryInstanceFromComponent(UInstancedStaticMeshComponent* InC
 				int32 BaseLightMapWidth = Instance->LODPerInstanceLightmapSize[LODIndex].X;
 				int32 BaseLightMapHeight = Instance->LODPerInstanceLightmapSize[LODIndex].Y;
 
-				FVector2D Scale = FVector2D(BaseLightMapWidth - 2, BaseLightMapHeight - 2) / Lightmap->Size;
+				FVector2D Scale = FVector2D(BaseLightMapWidth - 2, BaseLightMapHeight - 2) / Lightmap->GetPaddedSize();
 				Lightmap->LightmapObject->CoordinateScale = Scale;
 				Lightmap->LightmapObject->CoordinateBias = FVector2D(0, 0);
 
@@ -980,7 +980,7 @@ void FScene::AddGeometryInstanceFromComponent(UInstancedStaticMeshComponent* InC
 					{
 						int32 X = RenderIndex % InstancesPerRow;
 						int32 Y = RenderIndex / InstancesPerRow;
-						FVector2f Bias = (FVector2f(X, Y) * FVector2f(BaseLightMapWidth, BaseLightMapHeight) + FVector2f(1, 1)) / Lightmap->Size;
+						FVector2f Bias = (FVector2f(X, Y) * FVector2f(BaseLightMapWidth, BaseLightMapHeight) + FVector2f(1, 1)) / Lightmap->GetPaddedSize();
 						Lightmap->MeshMapBuildData->PerInstanceLightmapData[GameThreadInstanceIndex].LightmapUVBias = Bias;
 						Lightmap->MeshMapBuildData->PerInstanceLightmapData[GameThreadInstanceIndex].ShadowmapUVBias = Bias;
 					}
@@ -1167,7 +1167,7 @@ void FScene::AddGeometryInstanceFromComponent(ULandscapeComponent* InComponent)
 		{
 			Lightmap->CreateGameThreadResources();
 
-			Lightmap->LightmapObject->CoordinateScale = FVector2D(1, 1);
+			Lightmap->LightmapObject->CoordinateScale = FVector2D(Lightmap->Size) / Lightmap->GetPaddedSize();
 			Lightmap->LightmapObject->CoordinateBias = FVector2D(0, 0);
 
 			for (FDirectionalLightBuildInfo& DirectionalLight : LightScene.DirectionalLights.Elements)
@@ -1971,10 +1971,12 @@ void FScene::ApplyFinishedLightmapsToWorld()
 									Lightmap.TileStorage[Coords].CPUTextureData[0]->Decompress();
 									Lightmap.TileStorage[Coords].CPUTextureData[1]->Decompress();
 									Lightmap.TileStorage[Coords].CPUTextureData[2]->Decompress();
+									Lightmap.TileStorage[Coords].CPUTextureData[3]->Decompress();
 
 									Lightmap.TileStorage[ParentCoords].CPUTextureData[0]->Decompress();
 									Lightmap.TileStorage[ParentCoords].CPUTextureData[1]->Decompress();
 									Lightmap.TileStorage[ParentCoords].CPUTextureData[2]->Decompress();
+									Lightmap.TileStorage[ParentCoords].CPUTextureData[3]->Decompress();
 
 									for (int32 X = 0; X < GPreviewLightmapVirtualTileSize; X++)
 									{
@@ -1994,6 +1996,7 @@ void FScene::ApplyFinishedLightmapsToWorld()
 											Lightmap.TileStorage[Coords].CPUTextureData[0]->Data[DstLinearIndex] = Lightmap.TileStorage[ParentCoords].CPUTextureData[0]->Data[SrcLinearIndex];
 											Lightmap.TileStorage[Coords].CPUTextureData[1]->Data[DstLinearIndex] = Lightmap.TileStorage[ParentCoords].CPUTextureData[1]->Data[SrcLinearIndex];
 											Lightmap.TileStorage[Coords].CPUTextureData[2]->Data[DstLinearIndex] = Lightmap.TileStorage[ParentCoords].CPUTextureData[2]->Data[SrcLinearIndex];
+											Lightmap.TileStorage[Coords].CPUTextureData[3]->Data[DstLinearIndex] = Lightmap.TileStorage[ParentCoords].CPUTextureData[3]->Data[SrcLinearIndex];
 										}
 									}
 
@@ -2071,6 +2074,7 @@ void FScene::ApplyFinishedLightmapsToWorld()
 							Tile.Value.CPUTextureData[0]->Decompress();
 							Tile.Value.CPUTextureData[1]->Decompress();
 							Tile.Value.CPUTextureData[2]->Decompress();
+							Tile.Value.CPUTextureData[3]->Decompress();
 						}
 
 						// Transencode GI layers
@@ -2396,6 +2400,7 @@ void FScene::ApplyFinishedLightmapsToWorld()
 							Tile.Value.CPUTextureData[0]->Decompress();
 							Tile.Value.CPUTextureData[1]->Decompress();
 							Tile.Value.CPUTextureData[2]->Decompress();
+							Tile.Value.CPUTextureData[3]->Decompress();
 						}
 
 						FInstanceGroup& InstanceGroup = InstanceGroups.Elements[InstanceGroupIndex];
@@ -2737,6 +2742,7 @@ void FScene::ApplyFinishedLightmapsToWorld()
 							Tile.Value.CPUTextureData[0]->Decompress();
 							Tile.Value.CPUTextureData[1]->Decompress();
 							Tile.Value.CPUTextureData[2]->Decompress();
+							Tile.Value.CPUTextureData[3]->Decompress();
 						}
 
 						// Transencode GI layers
diff --git a/Engine/Plugins/Experimental/GPULightmass/Source/GPULightmass/Private/Scene/Scene.h b/Engine/Plugins/Experimental/GPULightmass/Source/GPULightmass/Private/Scene/Scene.h
--- a/Engine/Plugins/Experimental/GPULightmass/Source/GPULightmass/Private/Scene/Scene.h
+++ b/Engine/Plugins/Experimental/GPULightmass/Source/GPULightmass/Private/Scene/Scene.h
@@ -129,6 +129,9 @@ public:
 
 	void CalculateDistributionPrefixSumForAllLightmaps();
 
+	TArray<FLightmapRenderStateRef> MortonSortedLightmapRefList;
+	void BuildMortonSortedLightmapRefList();
+
 	volatile int32 Percentage = 0;
 };
 
diff --git a/Engine/Plugins/Experimental/GPULightmass/Source/GPULightmass/Private/Scene/StaticMesh.cpp b/Engine/Plugins/Experimental/GPULightmass/Source/GPULightmass/Private/Scene/StaticMesh.cpp
--- a/Engine/Plugins/Experimental/GPULightmass/Source/GPULightmass/Private/Scene/StaticMesh.cpp
+++ b/Engine/Plugins/Experimental/GPULightmass/Source/GPULightmass/Private/Scene/StaticMesh.cpp
@@ -55,9 +55,23 @@ void FStaticMeshInstance::AllocateLightmaps(TEntityArray<FLightmap>& LightmapCon
 			// Shrink LOD texture lightmaps by half for each LOD level
 			const int32 LightMapWidth = LODIndex > 0 ? FMath::Max(BaseLightMapWidth / (2 << (LODIndex - 1)), 32) : BaseLightMapWidth;
 			const int32 LightMapHeight = LODIndex > 0 ? FMath::Max(BaseLightMapHeight / (2 << (LODIndex - 1)), 32) : BaseLightMapHeight;
-
-			FString LightmapName = TEXT("Lightmap_") + (ComponentUObject->GetOwner() ? ComponentUObject->GetOwner()->GetActorLabel() : FString());
-
+			
+			FString LightmapName;
+			if (ComponentUObject->GetStaticMesh()->GetRenderData()->LODResources.Num() > 1)
+			{
+				LightmapName = FString::Printf(
+					TEXT("Lightmap_%s:LOD%d"),
+					*(ComponentUObject->GetOwner() ? ComponentUObject->GetOwner()->GetActorLabel() : FString()),
+					LODIndex);
+			}
+			else
+			{
+				LightmapName = FString::Printf(
+					TEXT("Lightmap_%s"), *(ComponentUObject->GetOwner()
+						                             ? ComponentUObject->GetOwner()->GetActorLabel()
+						                             : FString()));
+			}
+			
 			LODLightmaps.Add(LightmapContainer.Emplace(LightmapName, FIntPoint(LightMapWidth, LightMapHeight)));
 		}
 		else
diff --git a/Engine/Plugins/Experimental/GPULightmass/Source/GPULightmass/Private/VolumetricLightmap.cpp b/Engine/Plugins/Experimental/GPULightmass/Source/GPULightmass/Private/VolumetricLightmap.cpp
--- a/Engine/Plugins/Experimental/GPULightmass/Source/GPULightmass/Private/VolumetricLightmap.cpp
+++ b/Engine/Plugins/Experimental/GPULightmass/Source/GPULightmass/Private/VolumetricLightmap.cpp
@@ -555,6 +555,7 @@ void FVolumetricLightmapRenderer::VoxelizeScene()
 
 		{
 			FWriteSortedBrickRequestsCS::FParameters* PassParameters = GraphBuilder.AllocParameters<FWriteSortedBrickRequestsCS::FParameters>();
+			PassParameters->NumTotalBricks = NumTotalBricks;
 			PassParameters->BrickRequests = BrickRequests.UAV;
 			PassParameters->BrickRequestsUnsorted = GraphBuilder.CreateUAV(BrickRequestsUnsorted, PF_R32G32B32A32_UINT);
 			PassParameters->BrickRequestSortedIndices = GraphBuilder.CreateUAV(BrickRequestSortedIndices, PF_R32_UINT);
@@ -567,7 +568,7 @@ void FVolumetricLightmapRenderer::VoxelizeScene()
 				RDG_EVENT_NAME("WriteSortedBrickRequests"),
 				ComputeShader,
 				PassParameters,
-				FIntVector(NumTotalBricks, 1, 1));
+				FComputeShaderUtils::GetGroupCount(NumTotalBricks, 64));
 		}
 
 		for (int32 MipLevel = VoxelizationVolumeMips.Num() - 1; MipLevel >= 0; MipLevel--)
diff --git a/Engine/Plugins/Experimental/GPULightmass/Source/GPULightmass/Private/VolumetricLightmapVoxelization.h b/Engine/Plugins/Experimental/GPULightmass/Source/GPULightmass/Private/VolumetricLightmapVoxelization.h
--- a/Engine/Plugins/Experimental/GPULightmass/Source/GPULightmass/Private/VolumetricLightmapVoxelization.h
+++ b/Engine/Plugins/Experimental/GPULightmass/Source/GPULightmass/Private/VolumetricLightmapVoxelization.h
@@ -315,6 +315,7 @@ class FWriteSortedBrickRequestsCS : public FGlobalShader
 	}
 
 	BEGIN_SHADER_PARAMETER_STRUCT(FParameters, )
+		SHADER_PARAMETER(int32, NumTotalBricks)
 		SHADER_PARAMETER_RDG_BUFFER_UAV(RWBuffer<uint4>, BrickRequestsUnsorted)
 		SHADER_PARAMETER_RDG_BUFFER_UAV(RWBuffer<uint>, BrickRequestSortedIndices)
 		SHADER_PARAMETER_RDG_BUFFER_UAV(RWBuffer<uint>, BrickRequestSortedIndicesInverse)
