Fixed non unity error

diff --git a/Engine/Plugins/Experimental/GPULightmass/Source/GPULightmass/Private/GPULightmassSettings.cpp b/Engine/Plugins/Experimental/GPULightmass/Source/GPULightmass/Private/GPULightmassSettings.cpp
--- a/Engine/Plugins/Experimental/GPULightmass/Source/GPULightmass/Private/GPULightmassSettings.cpp
+++ b/Engine/Plugins/Experimental/GPULightmass/Source/GPULightmass/Private/GPULightmassSettings.cpp
@@ -13,6 +13,8 @@
 #include "MeshUtilitiesCommon.h"
 #include "Logging/MessageLog.h"
 #include "Misc/UObjectToken.h"
+#include "Misc/MessageDialog.h"
+
 
 #define LOCTEXT_NAMESPACE "StaticLightingSystem"
 
