GPULM: Add PATH_TRACING_PAYLOAD_OUTPUT_FLAG_TWO_SIDED_MATERIAL for invalid surface detection

diff --git a/Engine/Plugins/Experimental/GPULightmass/Shaders/Private/LightmapPathTracing.usf b/Engine/Plugins/Experimental/GPULightmass/Shaders/Private/LightmapPathTracing.usf
--- a/Engine/Plugins/Experimental/GPULightmass/Shaders/Private/LightmapPathTracing.usf
+++ b/Engine/Plugins/Experimental/GPULightmass/Shaders/Private/LightmapPathTracing.usf
@@ -170,6 +170,11 @@ void AccumulateRadiance(inout float3 TotalRadiance, float3 PathRadiance)
 	}
 }
 
+bool IsInvalidSurfaceHit(FPathTracingPayload Payload)
+{
+	return !Payload.IsFrontFace() && !Payload.IsMaterialTransmissive() && !Payload.IsMaterialTwoSided();
+}
+
 FPathTracingPayload TraceTransparentRay(FRayDesc Ray, bool IsCameraRay, bool LastBounce, bool IncludeEmission, uint2 LaunchIndex, uint NumLights, inout RandomSequence RandSequence, inout float3 PathThroughput, inout float3 Radiance)
 {
 	const uint RayFlags = IsCameraRay && EnableCameraBackfaceCulling ? RAY_FLAG_CULL_BACK_FACING_TRIANGLES : 0;
@@ -264,6 +269,13 @@ FPathTracingPayload TraceTransparentRay(FRayDesc Ray, bool IsCameraRay, bool Las
 			}
 		}
 
+		if (IsInvalidSurfaceHit(HitPayload))
+		{
+			// Propagate to Payload and terminate immediately upon invalid surface hits
+			Payload = HitPayload;
+			break;
+		}
+
 		// prepare next step around the loop
 		// retrace the exact same ray with TMin one ulp past the hit we just found
 		PathThroughput *= Transparency;
@@ -407,15 +419,15 @@ void PathTracingKernel(
 			Payload = TraceTransparentRay(Ray, false , bIsLastBounce, bIncludeEmissive, LaunchIndex, NumVisibleLights, RandSequence, PathThroughput, Radiance);
 		}
 		
-		if (Payload.IsMiss())
+		if (IsInvalidSurfaceHit(Payload))
 		{
-			// we didn't hit anything selectable for further shading, we are done
+			bIsValidSample = false;
 			break;
 		}
 		
-		if (Payload.Opacity < 0)
+		if (Payload.IsMiss())
 		{
-			bIsValidSample = false;
+			// we didn't hit anything selectable for further shading, we are done
 			break;
 		}
 
diff --git a/Engine/Shaders/Private/PathTracing/PathTracingCommon.ush b/Engine/Shaders/Private/PathTracing/PathTracingCommon.ush
--- a/Engine/Shaders/Private/PathTracing/PathTracingCommon.ush
+++ b/Engine/Shaders/Private/PathTracing/PathTracingCommon.ush
@@ -14,6 +14,7 @@
 // These flags are analogous to RAY_TRACING_PAYLOAD_* flags. There are currently 6 bits available
 #define PATH_TRACING_PAYLOAD_OUTPUT_FLAG_FRONT_FACE                   (1 << 0) // Indicates that ray has hit the front face of a primitive. Set by closest hit shader.
 #define PATH_TRACING_PAYLOAD_OUTPUT_FLAG_DECAL_RECEIVER               (1 << 1) // Indicates that ray has hit a primitive that receives decals. Set by closest hit shader.
+#define PATH_TRACING_PAYLOAD_OUTPUT_FLAG_TWO_SIDED_MATERIAL           (1 << 2) // Indicates that ray has hit a primitive surface with a two sided material. Set by closest hit shader. Used by GPU Lightmass to detect invalid surfaces.
 
 #define PATH_TRACING_PAYLOAD_INPUT_FLAG_RAY_TYPE_MASK                 (7 << 0) // Reserve 3 bits for ray-mask
 #define PATH_TRACING_PAYLOAD_INPUT_FLAG_RAY_TYPE_CAMERA                0       // Indicates that this is a camera ray which can trigger a number of special behaviors in the closest hit shader
@@ -167,7 +168,9 @@ struct FPathTracingPayload : FMinimalPayload
 
 	void SetDecalReceiver() { Flags |= PATH_TRACING_PAYLOAD_OUTPUT_FLAG_DECAL_RECEIVER; }
 	bool IsDecalReceiver() { return (Flags & PATH_TRACING_PAYLOAD_OUTPUT_FLAG_DECAL_RECEIVER) != 0; }
-
+	
+	void SetMaterialTwoSided() { Flags |= PATH_TRACING_PAYLOAD_OUTPUT_FLAG_TWO_SIDED_MATERIAL; }
+	bool IsMaterialTwoSided() { return (Flags & PATH_TRACING_PAYLOAD_OUTPUT_FLAG_TWO_SIDED_MATERIAL) != 0; }
 
 	// Various ways to interpret CustomData (depending on ShadingModelID)
 	// NOTE: This is not always following the same conventions as GetMaterialCustomData0,1()
diff --git a/Engine/Shaders/Private/PathTracing/PathTracingMaterialHitShader.usf b/Engine/Shaders/Private/PathTracing/PathTracingMaterialHitShader.usf
--- a/Engine/Shaders/Private/PathTracing/PathTracingMaterialHitShader.usf
+++ b/Engine/Shaders/Private/PathTracing/PathTracingMaterialHitShader.usf
@@ -571,6 +571,9 @@ RAY_TRACING_ENTRY_CLOSEST_HIT(PathTracingMaterialCHS,
 		// when viewing the surface from "inside", don't include emission
 		Payload.Radiance = 0;
 	}
+#endif
+#if MATERIAL_TWOSIDED
+	Payload.SetMaterialTwoSided();
 #endif
 	Payload.BaseColor = BaseColor;
 	Payload.Specular = 0;
@@ -584,14 +587,6 @@ RAY_TRACING_ENTRY_CLOSEST_HIT(PathTracingMaterialCHS,
 #endif
 	Payload.Ior = 0;
 
-#if !MATERIAL_TWOSIDED
-	if (!Payload.IsFrontFace() && !Payload.IsMaterialTransmissive())
-	{
-		// Overloaded semantics for opacity
-		Payload.Opacity = -1;
-	}
-#endif
-
 #endif
 
 	PackedPayload = PackPathTracingPayload(Payload);
