GPULM: Fix ray misses being misclassified as invalid hits

diff --git a/Engine/Plugins/Experimental/GPULightmass/Shaders/Private/LightmapPathTracing.usf b/Engine/Plugins/Experimental/GPULightmass/Shaders/Private/LightmapPathTracing.usf
--- a/Engine/Plugins/Experimental/GPULightmass/Shaders/Private/LightmapPathTracing.usf
+++ b/Engine/Plugins/Experimental/GPULightmass/Shaders/Private/LightmapPathTracing.usf
@@ -419,15 +419,15 @@ void PathTracingKernel(
 			Payload = TraceTransparentRay(Ray, false , bIsLastBounce, bIncludeEmissive, LaunchIndex, NumVisibleLights, RandSequence, PathThroughput, Radiance);
 		}
 		
-		if (IsInvalidSurfaceHit(Payload))
+		if (Payload.IsMiss())
 		{
-			bIsValidSample = false;
+			// we didn't hit anything selectable for further shading, we are done
 			break;
 		}
 		
-		if (Payload.IsMiss())
+		if (IsInvalidSurfaceHit(Payload))
 		{
-			// we didn't hit anything selectable for further shading, we are done
+			bIsValidSample = false;
 			break;
 		}
 
