Move UGPULightmassSubsystem setup logic to PostInitialize to prevent an issue with GetSettingsActor

diff --git a/Engine/Plugins/Experimental/GPULightmass/Source/GPULightmass/Private/GPULightmassSettings.cpp b/Engine/Plugins/Experimental/GPULightmass/Source/GPULightmass/Private/GPULightmassSettings.cpp
--- a/Engine/Plugins/Experimental/GPULightmass/Source/GPULightmass/Private/GPULightmassSettings.cpp
+++ b/Engine/Plugins/Experimental/GPULightmass/Source/GPULightmass/Private/GPULightmassSettings.cpp
@@ -118,7 +118,7 @@ bool UGPULightmassSettings::CanEditChange(const FProperty* InProperty) const
 }
 #endif
 
-void UGPULightmassSubsystem::Initialize(FSubsystemCollectionBase& Collection)
+void UGPULightmassSubsystem::PostInitialize()
 {
 	UWorld* World = Cast<UWorld>(GetOuter());
 	
diff --git a/Engine/Plugins/Experimental/GPULightmass/Source/GPULightmass/Public/GPULightmassSettings.h b/Engine/Plugins/Experimental/GPULightmass/Source/GPULightmass/Public/GPULightmassSettings.h
--- a/Engine/Plugins/Experimental/GPULightmass/Source/GPULightmass/Public/GPULightmassSettings.h
+++ b/Engine/Plugins/Experimental/GPULightmass/Source/GPULightmass/Public/GPULightmassSettings.h
@@ -154,7 +154,7 @@ class GPULIGHTMASS_API UGPULightmassSubsystem : public UWorldSubsystem
 {
 	GENERATED_BODY()
 public:
-	virtual void Initialize(FSubsystemCollectionBase& Collection) override;
+	virtual void PostInitialize() override;
 
 	UFUNCTION(BlueprintCallable, Category = GPULightmass)
 	UGPULightmassSettings* GetSettings();
