GPULM: Remove RAY_FLAG_SKIP_CLOSEST_HIT_SHADER as now the payload defaults to miss and a CHS is required to set a hit.

diff --git a/Engine/Plugins/Experimental/GPULightmass/Shaders/Private/LightmapPathTracing.usf b/Engine/Plugins/Experimental/GPULightmass/Shaders/Private/LightmapPathTracing.usf
--- a/Engine/Plugins/Experimental/GPULightmass/Shaders/Private/LightmapPathTracing.usf
+++ b/Engine/Plugins/Experimental/GPULightmass/Shaders/Private/LightmapPathTracing.usf
@@ -1420,7 +1420,7 @@ void VolumetricLightmapPathTracingMainRG()
 			
 		if (bIsValidRay)
 		{
-			uint RayFlags = RAY_FLAG_SKIP_CLOSEST_HIT_SHADER | RAY_FLAG_ACCEPT_FIRST_HIT_AND_END_SEARCH;
+			uint RayFlags = RAY_FLAG_ACCEPT_FIRST_HIT_AND_END_SEARCH;
 	
 			FMinimalPayload MinimalPayload = TraceLightmapVisibilityRay(
 				TLAS,
@@ -1505,7 +1505,7 @@ void StationaryLightShadowTracingMainRG()
 			
 			if (bRaySampleValid)
 			{
-				uint RayFlags = RAY_FLAG_SKIP_CLOSEST_HIT_SHADER | RAY_FLAG_ACCEPT_FIRST_HIT_AND_END_SEARCH;
+				uint RayFlags = RAY_FLAG_ACCEPT_FIRST_HIT_AND_END_SEARCH;
 
 				FMinimalPayload MinimalPayload = TraceLightmapVisibilityRay(
 					TLAS,
