GPULM: Fix SceneUniforms lifetime issue

diff --git a/Engine/Plugins/Experimental/GPULightmass/Source/GPULightmass/Private/LightmapRenderer.cpp b/Engine/Plugins/Experimental/GPULightmass/Source/GPULightmass/Private/LightmapRenderer.cpp
--- a/Engine/Plugins/Experimental/GPULightmass/Source/GPULightmass/Private/LightmapRenderer.cpp
+++ b/Engine/Plugins/Experimental/GPULightmass/Source/GPULightmass/Private/LightmapRenderer.cpp
@@ -780,7 +780,7 @@ void FCachedRayTracingSceneData::RestoreCachedBuffers(FRDGBuilder& GraphBuilder,
 {
 	TRACE_CPUPROFILER_EVENT_SCOPE(RestoreCachedBuffers);
 
-	FGPUSceneResourceParameters GPUScene{};
+	FGPUSceneResourceParameters GPUScene = RenderState.SceneUniforms.GetParameters().GPUScene;
 	
 	check(GPUScenePrimitiveDataBuffer.IsValid())
 	GPUScene.GPUScenePrimitiveSceneData = GraphBuilder.CreateSRV(GraphBuilder.RegisterExternalBuffer(GPUScenePrimitiveDataBuffer));
@@ -798,6 +798,9 @@ void FCachedRayTracingSceneData::RestoreCachedBuffers(FRDGBuilder& GraphBuilder,
 	check(GPUSceneLightDataBuffer.IsValid())
 	GPUScene.GPUSceneLightData = GraphBuilder.CreateSRV(GraphBuilder.RegisterExternalBuffer(GPUSceneLightDataBuffer));
 
+	// Clear cached parameters so that bGPUSceneIsDirty will be true
+	// Required as the uniform buffer in SceneUniforms has only 1 frame lifetime and needs to be recreated
+	RenderState.SceneUniforms.Set({});
 	RenderState.SceneUniforms.Set(GPUScene);
 }
 
