Removing a few more dependencies on Renderer/Private code.

diff --git a/Engine/Plugins/Experimental/GPULightmass/Source/GPULightmass/Private/Scene/Scene.cpp b/Engine/Plugins/Experimental/GPULightmass/Source/GPULightmass/Private/Scene/Scene.cpp
--- a/Engine/Plugins/Experimental/GPULightmass/Source/GPULightmass/Private/Scene/Scene.cpp
+++ b/Engine/Plugins/Experimental/GPULightmass/Source/GPULightmass/Private/Scene/Scene.cpp
@@ -29,7 +29,7 @@
 #include "Misc/FileHelper.h"
 #include "Components/ReflectionCaptureComponent.h"
 #include "ReflectionEnvironment.h"
-#include "ScenePrivate.h"
+#include "RHIStaticStates.h"
 
 #define LOCTEXT_NAMESPACE "StaticLightingSystem"
 
@@ -434,7 +434,7 @@ void FScene::AddLight(LightComponentType* LightComponent)
 			FPrimitiveSceneInfo* PrimitiveSceneInfo = SceneProxy->GetPrimitiveSceneInfo();
 			if (PrimitiveSceneInfo && PrimitiveSceneInfo->IsIndexValid())
 			{
-				PrimitiveSceneInfo->Scene->GPUScene.AddPrimitiveToUpdate(PrimitiveSceneInfo->GetIndex(), EPrimitiveDirtyState::ChangedStaticLighting);
+				PrimitiveSceneInfo->RequestGPUSceneUpdate(EPrimitiveDirtyState::ChangedStaticLighting);
 			}
 		}
 	});
@@ -531,7 +531,7 @@ void FScene::RemoveLight(LightComponentType* PointLightComponent)
 			FPrimitiveSceneInfo* PrimitiveSceneInfo = SceneProxy->GetPrimitiveSceneInfo();
 			if (PrimitiveSceneInfo && PrimitiveSceneInfo->IsIndexValid())
 			{
-				PrimitiveSceneInfo->Scene->GPUScene.AddPrimitiveToUpdate(PrimitiveSceneInfo->GetIndex(), EPrimitiveDirtyState::ChangedStaticLighting);
+				PrimitiveSceneInfo->RequestGPUSceneUpdate(EPrimitiveDirtyState::ChangedStaticLighting);
 			}
 		}
 	});
