Changine more code to use the new DrawRectangle interface.

diff --git a/Engine/Plugins/Experimental/GPULightmass/Source/GPULightmass/Private/LightmapRenderer.cpp b/Engine/Plugins/Experimental/GPULightmass/Source/GPULightmass/Private/LightmapRenderer.cpp
--- a/Engine/Plugins/Experimental/GPULightmass/Source/GPULightmass/Private/LightmapRenderer.cpp
+++ b/Engine/Plugins/Experimental/GPULightmass/Source/GPULightmass/Private/LightmapRenderer.cpp
@@ -31,7 +31,7 @@
 #include "Engine/SubsurfaceProfile.h"
 #include "IESTextureManager.h"
 #include "PathTracing.h"
-#include "PostProcess/SceneFilterRendering.h"
+#include "PostProcess/DrawRectangle.h"
 #include "SceneUniformBuffer.h"
 
 RENDERER_API uint8 BlendModeToRayTracingInstanceMask(const EBlendMode BlendMode, ERayTracingViewMaskMode MaskMode);
@@ -3903,15 +3903,16 @@ void FLightmapRenderer::RenderIrradianceCacheVisualization(FPostOpaqueRenderPara
 
 		SetShaderParameters(RHICmdList, PixelShader, PixelShader.GetPixelShader(), *PassParameters);
 
-		DrawRectangle(
+		UE::Renderer::PostProcess::DrawRectangle(
 			RHICmdList,
+			VertexShader,
 			0, 0,
 			ViewportRect.Width(), ViewportRect.Height(),
 			0, 0,
 			ViewportRect.Width(), ViewportRect.Height(),
 			FIntPoint(ViewportRect.Width(), ViewportRect.Height()),
-			TextureExtent,
-			VertexShader);
+			TextureExtent
+		);
 	});
 }
 
