GPULM: Fix GPUScene UB not bound when dispatching updates for Landscape.

diff --git a/Engine/Plugins/Experimental/GPULightmass/Source/GPULightmass/Private/LightmapRenderer.cpp b/Engine/Plugins/Experimental/GPULightmass/Source/GPULightmass/Private/LightmapRenderer.cpp
--- a/Engine/Plugins/Experimental/GPULightmass/Source/GPULightmass/Private/LightmapRenderer.cpp
+++ b/Engine/Plugins/Experimental/GPULightmass/Source/GPULightmass/Private/LightmapRenderer.cpp
@@ -1053,7 +1053,7 @@ bool FSceneRenderState::SetupRayTracingScene(FRDGBuilder& GraphBuilder, int32 LO
 							FRHIResourceCreateInfo ScratchBufferCreateInfo(TEXT("RHILandscapeScratchBuffer"));
 							FBufferRHIRef ScratchBuffer = RHICreateBuffer(BLASScratchSize, BUF_StructuredBuffer | BUF_RayTracingScratch, 0, ERHIAccess::UAVCompute, ScratchBufferCreateInfo);
 
-							RHICmdList.SetStaticUniformBuffers({ReferenceView->ViewUniformBuffer});
+							RHICmdList.SetStaticUniformBuffers({ReferenceView->ViewUniformBuffer, SceneUniforms.GetBufferRHI(GraphBuilder)});
 							DynamicGeometryCollection.DispatchUpdates(RHICmdList, ScratchBuffer);
 
 							// Landscape VF doesn't really use the vertex buffer in HitGroupSystemParameters
