GPULM: Show and hide progress notification based on percentage to avoid user confusion when in BakeWhatYouSee mode

diff --git a/Engine/Plugins/Experimental/GPULightmass/Source/GPULightmass/Private/GPULightmass.cpp b/Engine/Plugins/Experimental/GPULightmass/Source/GPULightmass/Private/GPULightmass.cpp
--- a/Engine/Plugins/Experimental/GPULightmass/Source/GPULightmass/Private/GPULightmass.cpp
+++ b/Engine/Plugins/Experimental/GPULightmass/Source/GPULightmass/Private/GPULightmass.cpp
@@ -26,36 +26,55 @@ FGPULightmass::FGPULightmass(UWorld* InWorld, FGPULightmassModule* InGPULightmas
 
 	SettingsGuard = MakeUnique<FGCObjectScopeGuard>(Settings);
 
-	// Start the lightmass 'progress' notification
-	FNotificationInfo Info(LOCTEXT("LightBuildMessage", "Building lighting"));
-	Info.bFireAndForget = false;
+	StartTime = FPlatformTime::Seconds();
+}
 
-	if (InSettings->Mode == EGPULightmassMode::BakeWhatYouSee)
+void FGPULightmass::ShowLightBuildNotification()
+{
+	if (!IsEngineExitRequested() && !LightBuildNotification.IsValid())
 	{
+		// Start the lightmass 'progress' notification
+		FNotificationInfo Info(LOCTEXT("LightBuildMessage", "Building lighting"));
+		Info.bFireAndForget = false;
+
+		if (Settings->Mode == EGPULightmassMode::BakeWhatYouSee)
+		{
+			Info.ButtonDetails.Add(FNotificationButtonInfo(
+			LOCTEXT("Save", "Save and Stop"),
+			FText::GetEmpty(),
+			FSimpleDelegate::CreateLambda([World = World, this]() { 
+				this->Scene.ApplyFinishedLightmapsToWorld(); 
+				World->GetSubsystem<UGPULightmassSubsystem>()->Stop(); 
+			})));
+		}
+	
 		Info.ButtonDetails.Add(FNotificationButtonInfo(
-		LOCTEXT("Save", "Save and Stop"),
-		FText::GetEmpty(),
-		FSimpleDelegate::CreateLambda([InWorld, this]() { 
-			this->Scene.ApplyFinishedLightmapsToWorld(); 
-			InWorld->GetSubsystem<UGPULightmassSubsystem>()->Stop(); 
-		})));
+			LOCTEXT("LightBuildCancel", "Cancel"),
+			LOCTEXT("LightBuildCancelToolTip", "Cancels the lighting build in progress."),
+			FSimpleDelegate::CreateLambda([World = World]() { 
+				World->GetSubsystem<UGPULightmassSubsystem>()->Stop(); 
+				World->GetSubsystem<UGPULightmassSubsystem>()->OnLightBuildEnded().Broadcast();
+				})));
+
+		LightBuildNotification = FSlateNotificationManager::Get().AddNotification(Info);
+		if (LightBuildNotification.IsValid())
+		{
+			LightBuildNotification->SetCompletionState(SNotificationItem::CS_Pending);
+		}
 	}
-	
-	Info.ButtonDetails.Add(FNotificationButtonInfo(
-		LOCTEXT("LightBuildCancel", "Cancel"),
-		LOCTEXT("LightBuildCancelToolTip", "Cancels the lighting build in progress."),
-		FSimpleDelegate::CreateLambda([InWorld]() { 
-			InWorld->GetSubsystem<UGPULightmassSubsystem>()->Stop(); 
-			InWorld->GetSubsystem<UGPULightmassSubsystem>()->OnLightBuildEnded().Broadcast();
-			})));
+}
 
-	LightBuildNotification = FSlateNotificationManager::Get().AddNotification(Info);
+void FGPULightmass::RemoveLightBuildNotification()
+{
 	if (LightBuildNotification.IsValid())
 	{
-		LightBuildNotification->SetCompletionState(SNotificationItem::CS_Pending);
-	}
-
-	StartTime = FPlatformTime::Seconds();
+		// Shows a notification that fades out slowly
+		FText CompletedText = LOCTEXT("LightBuildDoneMessage", "Lighting build completed");
+		LightBuildNotification->SetText(CompletedText);
+		LightBuildNotification->SetCompletionState(SNotificationItem::CS_Success);
+		LightBuildNotification->ExpireAndFadeout();
+		LightBuildNotification.Reset();
+	}	
 }
 
 void FGPULightmass::GameThreadDestroy()
@@ -65,23 +84,7 @@ void FGPULightmass::GameThreadDestroy()
 	UE_LOG(LogGPULightmass, Log, TEXT("Total lighting time: %s"), *FPlatformTime::PrettyTime(FPlatformTime::Seconds() - StartTime));
 
 	RemoveGameThreadEventHooks();
-
-	if (LightBuildNotification.IsValid())
-	{
-		if (!IsEngineExitRequested())
-		{
-			// Shows a notification that fades out slowly
-			FText CompletedText = LOCTEXT("LightBuildDoneMessage", "Lighting build completed");
-			LightBuildNotification->SetText(CompletedText);
-			LightBuildNotification->SetCompletionState(SNotificationItem::CS_Success);
-			LightBuildNotification->ExpireAndFadeout();
-		}
-		else
-		{
-			// Immediately destroys the notification widget to avoid crash
-			LightBuildNotification.Reset();
-		}
-	}
+	RemoveLightBuildNotification();
 
 	Scene.RemoveAllComponents();
 }
diff --git a/Engine/Plugins/Experimental/GPULightmass/Source/GPULightmass/Private/GPULightmass.h b/Engine/Plugins/Experimental/GPULightmass/Source/GPULightmass/Private/GPULightmass.h
--- a/Engine/Plugins/Experimental/GPULightmass/Source/GPULightmass/Private/GPULightmass.h
+++ b/Engine/Plugins/Experimental/GPULightmass/Source/GPULightmass/Private/GPULightmass.h
@@ -30,6 +30,9 @@ public:
 	void InstallGameThreadEventHooks();
 	void RemoveGameThreadEventHooks();
 
+	void ShowLightBuildNotification();
+	void RemoveLightBuildNotification();
+
 	void StartRecordingVisibleTiles();
 	void EndRecordingVisibleTiles();
 
diff --git a/Engine/Plugins/Experimental/GPULightmass/Source/GPULightmass/Private/Scene/Scene.cpp b/Engine/Plugins/Experimental/GPULightmass/Source/GPULightmass/Private/Scene/Scene.cpp
--- a/Engine/Plugins/Experimental/GPULightmass/Source/GPULightmass/Private/Scene/Scene.cpp
+++ b/Engine/Plugins/Experimental/GPULightmass/Source/GPULightmass/Private/Scene/Scene.cpp
@@ -1491,6 +1491,15 @@ void FScene::BackgroundTick()
 {
 	int32 Percentage = FPlatformAtomics::AtomicRead(&RenderState.Percentage);
 
+	if (Percentage != 100)
+	{
+		GPULightmass->ShowLightBuildNotification();
+	}
+	else
+	{
+		GPULightmass->RemoveLightBuildNotification();
+	}
+	
 	if (GPULightmass->LightBuildNotification.IsValid())
 	{
 		bool bIsViewportNonRealtime = GCurrentLevelEditingViewportClient && !GCurrentLevelEditingViewportClient->IsRealtime();
