GPULM: Ignore non-static actors when checking for outdate lightmap UV versions

diff --git a/Engine/Plugins/Experimental/GPULightmass/Source/GPULightmass/Private/GPULightmassSettings.cpp b/Engine/Plugins/Experimental/GPULightmass/Source/GPULightmass/Private/GPULightmassSettings.cpp
--- a/Engine/Plugins/Experimental/GPULightmass/Source/GPULightmass/Private/GPULightmassSettings.cpp
+++ b/Engine/Plugins/Experimental/GPULightmass/Source/GPULightmass/Private/GPULightmassSettings.cpp
@@ -147,7 +147,10 @@ bool CheckStaticMeshesForLightmapUVVersionUpgrade(const UWorld* World)
 			
 	for (const UStaticMeshComponent* Component : TObjectRange<UStaticMeshComponent>(RF_ClassDefaultObject | RF_ArchetypeObject, true, EInternalObjectFlags::Garbage))
 	{
-		if (Component->GetWorld() == World && Component->GetStaticMesh() && Component->GetStaticMesh()->GetLightmapUVVersion() != (int32)ELightmapUVVersion::Latest)
+		if (Component->GetWorld() == World &&
+			Component->HasValidSettingsForStaticLighting(false) &&
+			Component->GetStaticMesh() &&
+			Component->GetStaticMesh()->GetLightmapUVVersion() != (int32)ELightmapUVVersion::Latest)
 		{
 			StaticMeshesToUpgradeLightmapUVVersion.Add(Component->GetStaticMesh().Get());
 		}
