Disable GPULM when using a shader preview mode.

diff --git a/Engine/Plugins/Experimental/GPULightmass/Source/GPULightmassEditor/Private/GPULightmassEditorModule.cpp b/Engine/Plugins/Experimental/GPULightmass/Source/GPULightmassEditor/Private/GPULightmassEditorModule.cpp
--- a/Engine/Plugins/Experimental/GPULightmass/Source/GPULightmassEditor/Private/GPULightmassEditorModule.cpp
+++ b/Engine/Plugins/Experimental/GPULightmass/Source/GPULightmassEditor/Private/GPULightmassEditorModule.cpp
@@ -77,9 +77,9 @@ enum class ERayTracingDisabledReason : int32
 typedef TBitArray<TInlineAllocator<EShaderPlatform::SP_NumPlatforms / 8>> ShaderPlatformMaskType;
 extern RENDERCORE_API ShaderPlatformMaskType GRayTracingPlatformMask;
 
-ERayTracingDisabledReason GetRayTracingDisabledReason()
+ERayTracingDisabledReason GetRayTracingDisabledReason(EShaderPlatform ShaderPlatform)
 {
-	if (IsRayTracingEnabled())
+	if (IsRayTracingEnabled(ShaderPlatform))
 	{
 		return ERayTracingDisabledReason::OK;
 	}
@@ -90,12 +90,12 @@ ERayTracingDisabledReason GetRayTracingDisabledReason()
 		return ERayTracingDisabledReason::DISABLED_BY_PROJECT_SETTINGS;
 	}
 
-	if (!RHISupportsRayTracing(GMaxRHIShaderPlatform))
+	if (!RHISupportsRayTracing(ShaderPlatform))
 	{
 		return ERayTracingDisabledReason::INCOMPATIBLE_SHADER_PLATFORM;
 	}
 	// Shader platform statically supports ray tracing, but disabled by target platform at runtime
-	else if (RHISupportsRayTracing(GMaxRHIShaderPlatform) && (GRayTracingPlatformMask[(int)GMaxRHIShaderPlatform]) == 0)
+	else if (RHISupportsRayTracing(ShaderPlatform) && (GRayTracingPlatformMask[(int)ShaderPlatform]) == 0)
 	{
 		return ERayTracingDisabledReason::DISABLED_BY_TARGET_PLATFORM;
 	}
@@ -200,12 +200,18 @@ TSharedRef<SDockTab> FGPULightmassEditorModule::SpawnSettingsTab(const FSpawnTab
 
 	SettingsView = PropPlugin.CreateDetailView(DetailsViewArgs);
 
+	ERayTracingDisabledReason RayTracingStatus = ERayTracingDisabledReason::INCAPABLE_HARDWARE;
 	if (UWorld* World = GEditor->GetEditorWorldContext().World())
 	{
 		if (World->GetSubsystem<UGPULightmassSubsystem>())
 		{
 			SettingsView->SetObject(World->GetSubsystem<UGPULightmassSubsystem>()->GetSettings());
 		}
+
+		if (World->Scene)
+		{
+			RayTracingStatus = GetRayTracingDisabledReason(World->Scene->GetShaderPlatform());			
+		}
 	}
 
 	return SNew(SDockTab)
@@ -226,7 +232,7 @@ TSharedRef<SDockTab> FGPULightmassEditorModule::SpawnSettingsTab(const FSpawnTab
 					SNew(SButton)
 					.HAlign(HAlign_Center)
 					.ButtonStyle(FAppStyle::Get(), "FlatButton.Success")
-					.IsEnabled(IsRayTracingEnabled() && IsStaticLightingAllowed() && IsPathTracingEnabled())
+					.IsEnabled( (RayTracingStatus == ERayTracingDisabledReason::OK) && IsStaticLightingAllowed() && IsPathTracingEnabled())
 					.Visibility_Lambda([](){ return IsRunning() ? EVisibility::Collapsed : EVisibility::Visible; })
 					.OnClicked_Raw(this, &FGPULightmassEditorModule::OnStartClicked)
 					[
@@ -341,11 +347,11 @@ TSharedRef<SDockTab> FGPULightmassEditorModule::SpawnSettingsTab(const FSpawnTab
 			[
 				SAssignNew(Messages, STextBlock)
 				.AutoWrapText(true)
-				.Text_Lambda( []() -> FText
+				.Text_Lambda( [RayTracingStatus]() -> FText
 				{
-					if (!IsRayTracingEnabled())
+					if (RayTracingStatus != ERayTracingDisabledReason::OK)
 					{
-						return GenerateRayTracingDisabledReasonMessage(GetRayTracingDisabledReason());
+						return GenerateRayTracingDisabledReasonMessage(RayTracingStatus);
 					}
 					else if (!IsStaticLightingAllowed())
 					{
