GPULM: Check ray tracing status all the time in case preview mode changes.

diff --git a/Engine/Plugins/Experimental/GPULightmass/Source/GPULightmassEditor/Private/GPULightmassEditorModule.cpp b/Engine/Plugins/Experimental/GPULightmass/Source/GPULightmassEditor/Private/GPULightmassEditorModule.cpp
--- a/Engine/Plugins/Experimental/GPULightmass/Source/GPULightmassEditor/Private/GPULightmassEditorModule.cpp
+++ b/Engine/Plugins/Experimental/GPULightmass/Source/GPULightmassEditor/Private/GPULightmassEditorModule.cpp
@@ -114,6 +114,19 @@ ERayTracingDisabledReason GetRayTracingDisabledReason(EShaderPlatform ShaderPlat
 	return ERayTracingDisabledReason::INCAPABLE_HARDWARE;
 }
 
+ERayTracingDisabledReason GetRayTracingDisabledReason()
+{
+	ERayTracingDisabledReason RayTracingStatus = ERayTracingDisabledReason::INCAPABLE_HARDWARE;
+	if (UWorld* World = GEditor->GetEditorWorldContext().World())
+	{
+		if (World->Scene)
+		{
+			RayTracingStatus = GetRayTracingDisabledReason(World->Scene->GetShaderPlatform());
+		}
+	}
+	return RayTracingStatus;
+}
+
 static FText GenerateRayTracingDisabledReasonMessage(ERayTracingDisabledReason Reason)
 {
 	switch (Reason)
@@ -200,18 +213,12 @@ TSharedRef<SDockTab> FGPULightmassEditorModule::SpawnSettingsTab(const FSpawnTab
 
 	SettingsView = PropPlugin.CreateDetailView(DetailsViewArgs);
 
-	ERayTracingDisabledReason RayTracingStatus = ERayTracingDisabledReason::INCAPABLE_HARDWARE;
 	if (UWorld* World = GEditor->GetEditorWorldContext().World())
 	{
 		if (World->GetSubsystem<UGPULightmassSubsystem>())
 		{
 			SettingsView->SetObject(World->GetSubsystem<UGPULightmassSubsystem>()->GetSettings());
 		}
-
-		if (World->Scene)
-		{
-			RayTracingStatus = GetRayTracingDisabledReason(World->Scene->GetShaderPlatform());			
-		}
 	}
 
 	return SNew(SDockTab)
@@ -232,7 +239,9 @@ TSharedRef<SDockTab> FGPULightmassEditorModule::SpawnSettingsTab(const FSpawnTab
 					SNew(SButton)
 					.HAlign(HAlign_Center)
 					.ButtonStyle(FAppStyle::Get(), "FlatButton.Success")
-					.IsEnabled( (RayTracingStatus == ERayTracingDisabledReason::OK) && IsStaticLightingAllowed() && IsPathTracingEnabled())
+					.IsEnabled_Lambda( []() {
+						return (GetRayTracingDisabledReason() == ERayTracingDisabledReason::OK) && IsStaticLightingAllowed() && IsPathTracingEnabled();
+					}) 
 					.Visibility_Lambda([](){ return IsRunning() ? EVisibility::Collapsed : EVisibility::Visible; })
 					.OnClicked_Raw(this, &FGPULightmassEditorModule::OnStartClicked)
 					[
@@ -347,8 +356,9 @@ TSharedRef<SDockTab> FGPULightmassEditorModule::SpawnSettingsTab(const FSpawnTab
 			[
 				SAssignNew(Messages, STextBlock)
 				.AutoWrapText(true)
-				.Text_Lambda( [RayTracingStatus]() -> FText
+				.Text_Lambda( []() -> FText
 				{
+					ERayTracingDisabledReason RayTracingStatus = GetRayTracingDisabledReason();
 					if (RayTracingStatus != ERayTracingDisabledReason::OK)
 					{
 						return GenerateRayTracingDisabledReasonMessage(RayTracingStatus);
