Removing -O0 workaround

diff --git a/Engine/Plugins/Experimental/GPULightmass/Source/GPULightmass/Private/LightmapRayTracing.h b/Engine/Plugins/Experimental/GPULightmass/Source/GPULightmass/Private/LightmapRayTracing.h
--- a/Engine/Plugins/Experimental/GPULightmass/Source/GPULightmass/Private/LightmapRayTracing.h
+++ b/Engine/Plugins/Experimental/GPULightmass/Source/GPULightmass/Private/LightmapRayTracing.h
@@ -63,8 +63,6 @@ class FLightmapPathTracingRGS : public FGlobalShader
 		OutEnvironment.SetDefine(TEXT("USE_RECT_LIGHT_TEXTURES"), 1);
 		OutEnvironment.CompilerFlags.Add(CFLAG_ForceDXC);
 		OutEnvironment.CompilerFlags.Add(CFLAG_WarningsAsErrors);
-		// @todo - Working around DXC compiler crash UE-165154, should be removed after DXC is updated with a fix.
-		OutEnvironment.CompilerFlags.Add(CFLAG_ForceOptimization);
 	}
 
 	static ERayTracingPayloadType GetRayTracingPayloadType(const int32 PermutationId)
@@ -129,8 +127,6 @@ class FVolumetricLightmapPathTracingRGS : public FGlobalShader
 		OutEnvironment.SetDefine(TEXT("USE_RECT_LIGHT_TEXTURES"), 1);
 		OutEnvironment.CompilerFlags.Add(CFLAG_ForceDXC);
 		OutEnvironment.CompilerFlags.Add(CFLAG_WarningsAsErrors);
-		// @todo - Working around DXC compiler crash UE-165154, should be removed after DXC is updated with a fix.
-		OutEnvironment.CompilerFlags.Add(CFLAG_ForceOptimization);
 	}
 
 	static ERayTracingPayloadType GetRayTracingPayloadType(const int32 PermutationId)
diff --git a/Engine/Source/Runtime/Renderer/Private/PathTracing.cpp b/Engine/Source/Runtime/Renderer/Private/PathTracing.cpp
--- a/Engine/Source/Runtime/Renderer/Private/PathTracing.cpp
+++ b/Engine/Source/Runtime/Renderer/Private/PathTracing.cpp
@@ -863,8 +863,6 @@ class FPathTracingRG : public FGlobalShader
 		FGlobalShader::ModifyCompilationEnvironment(Parameters, OutEnvironment);
 		OutEnvironment.SetDefine(TEXT("USE_RECT_LIGHT_TEXTURES"), 1);
 		OutEnvironment.CompilerFlags.Add(CFLAG_WarningsAsErrors);
-		// @todo - Working around DXC compiler crash UE-165154, should be removed after DXC is updated with a fix.
-		OutEnvironment.CompilerFlags.Add(CFLAG_ForceOptimization);
 	}
 
 	static ERayTracingPayloadType GetRayTracingPayloadType(const int32 PermutationId)
