GPULM: Set a dummy LooseParameters UB for instanced meshes to prevent a crash.

diff --git a/Engine/Plugins/Experimental/GPULightmass/Source/GPULightmass/Private/Scene/InstancedStaticMesh.cpp b/Engine/Plugins/Experimental/GPULightmass/Source/GPULightmass/Private/Scene/InstancedStaticMesh.cpp
--- a/Engine/Plugins/Experimental/GPULightmass/Source/GPULightmass/Private/Scene/InstancedStaticMesh.cpp
+++ b/Engine/Plugins/Experimental/GPULightmass/Source/GPULightmass/Private/Scene/InstancedStaticMesh.cpp
@@ -130,6 +130,9 @@ TArray<FMeshBatch> FInstanceGroupRenderState::GetMeshBatchesForGBufferRendering(
 {
 	TArray<FMeshBatch> MeshBatches;
 
+	FInstancedStaticMeshVFLooseUniformShaderParameters LooseParameters{};
+	FUniformBufferRHIRef LooseParametersUniformBuffer = CreateUniformBufferImmediate(LooseParameters, EUniformBufferUsage::UniformBuffer_SingleFrame);
+
 	// TODO: potential race conditions between GT & RT everywhere in the following code
 	FStaticMeshLODResources& LODModel = RenderData->LODResources[LODIndex];
 	for (int32 SectionIndex = 0; SectionIndex < LODModel.Sections.Num(); SectionIndex++)
@@ -166,6 +169,7 @@ TArray<FMeshBatch> FInstanceGroupRenderState::GetMeshBatchesForGBufferRendering(
 		MeshBatchElement.MinVertexIndex = Section.MinVertexIndex;
 		MeshBatchElement.MaxVertexIndex = Section.MaxVertexIndex;
 		MeshBatchElement.PrimitiveIdMode = PrimID_DynamicPrimitiveShaderData;
+		MeshBatchElement.LooseParametersUniformBuffer = LooseParametersUniformBuffer;
 
 		MeshBatchElement.InstancedLODIndex = LODIndex;
 		MeshBatch.LODIndex = LODIndex;
