Add analytics to GPU Lightmass

diff --git a/Engine/Plugins/Experimental/GPULightmass/Source/GPULightmass/Private/GPULightmassSettings.cpp b/Engine/Plugins/Experimental/GPULightmass/Source/GPULightmass/Private/GPULightmassSettings.cpp
--- a/Engine/Plugins/Experimental/GPULightmass/Source/GPULightmass/Private/GPULightmassSettings.cpp
+++ b/Engine/Plugins/Experimental/GPULightmass/Source/GPULightmass/Private/GPULightmassSettings.cpp
@@ -14,6 +14,7 @@
 #include "Logging/MessageLog.h"
 #include "Misc/UObjectToken.h"
 #include "Misc/MessageDialog.h"
+#include "EngineAnalytics.h"
 
 
 #define LOCTEXT_NAMESPACE "StaticLightingSystem"
@@ -275,6 +276,11 @@ void UGPULightmassSubsystem::Launch()
 						SubSlowTask.EnterProgressFrame(1, LOCTEXT("RegisteringComponentsWithStaticLightingSystem", "Registering components with static lighting system"));
 					}
 				}
+
+				if (FEngineAnalytics::IsAvailable())
+				{
+					FEngineAnalytics::GetProvider().RecordEvent(TEXT("GPULightmass.Capture"));
+				}
 			}
 			else
 			{
