GPULM: Fix stationary sky light - bent normal needs to be accumulated over all samples.

diff --git a/Engine/Plugins/Experimental/GPULightmass/Shaders/Private/LightmapPathTracing.usf b/Engine/Plugins/Experimental/GPULightmass/Shaders/Private/LightmapPathTracing.usf
--- a/Engine/Plugins/Experimental/GPULightmass/Shaders/Private/LightmapPathTracing.usf
+++ b/Engine/Plugins/Experimental/GPULightmass/Shaders/Private/LightmapPathTracing.usf
@@ -1359,7 +1359,7 @@ void VolumetricLightmapPathTracingMainRG()
 				
 			if (HalfSampleIndex == 0)
 			{
-				SkyBentNormal[VoxelPos].rgb = SkyLightBentNormal;
+				SkyBentNormal[VoxelPos].rgb += SkyLightBentNormal;
 			}
 		
 		#ifndef LIGHTMAP_PATH_TRACING_MAIN_RG
