Fixed build break.

diff --git a/Engine/Plugins/Experimental/GPULightmass/Source/GPULightmass/Private/LightmapRenderer.cpp b/Engine/Plugins/Experimental/GPULightmass/Source/GPULightmass/Private/LightmapRenderer.cpp
--- a/Engine/Plugins/Experimental/GPULightmass/Source/GPULightmass/Private/LightmapRenderer.cpp
+++ b/Engine/Plugins/Experimental/GPULightmass/Source/GPULightmass/Private/LightmapRenderer.cpp
@@ -1045,6 +1045,7 @@ bool FSceneRenderState::SetupRayTracingScene(FRDGBuilder& GraphBuilder, int32 LO
 							};
 
 							DynamicGeometryCollection.AddDynamicMeshBatchForGeometryUpdate(
+								RHICmdList,
 								Landscape.ComponentUObject->GetWorld()->Scene->GetRenderScene(),
 								ReferenceView.Get(),
 								nullptr,
@@ -1593,6 +1594,8 @@ void FLightmapRenderer::Finalize(FRDGBuilder& GraphBuilder)
 		return;
 	}
 
+	FRHICommandListImmediate& RHICmdList = GraphBuilder.RHICmdList;
+
 	const auto HoldReference = [](FRDGBuilder& GraphBuilder, const FShaderResourceViewRHIRef& View) -> const FShaderResourceViewRHIRef&
 	{
 		return *GraphBuilder.AllocObject<FShaderResourceViewRHIRef>(View);
@@ -1610,7 +1613,6 @@ void FLightmapRenderer::Finalize(FRDGBuilder& GraphBuilder)
 
 		if (TileUploadRequests.Num() > 0)
 		{
-			FRHICommandListImmediate& RHICmdList = GraphBuilder.RHICmdList;
 			SCOPED_DRAW_EVENTF(RHICmdList, GPULightmassUploadConvergedTiles, TEXT("GPULightmass UploadConvergedTiles %d tiles"), TileUploadRequests.Num());
 
 			int32 NewSize = FMath::CeilToInt(FMath::Sqrt(static_cast<float>(TileUploadRequests.Num())));
diff --git a/Engine/Plugins/Media/PixelCapture/Source/PixelCapture/Private/PixelCaptureCapturerRHIToI420Compute.cpp b/Engine/Plugins/Media/PixelCapture/Source/PixelCapture/Private/PixelCaptureCapturerRHIToI420Compute.cpp
--- a/Engine/Plugins/Media/PixelCapture/Source/PixelCapture/Private/PixelCaptureCapturerRHIToI420Compute.cpp
+++ b/Engine/Plugins/Media/PixelCapture/Source/PixelCapture/Private/PixelCaptureCapturerRHIToI420Compute.cpp
@@ -38,7 +38,7 @@ FPixelCaptureCapturerRHIToI420Compute ::~FPixelCaptureCapturerRHIToI420Compute()
 
 void FPixelCaptureCapturerRHIToI420Compute::Initialize(int32 InputWidth, int32 InputHeight)
 {
-	FRHICommandListBase& RHICmdList = FRHICommandListImmediate::Get();
+	FRHICommandListImmediate& RHICmdList = FRHICommandListImmediate::Get();
 
 	const int32 Width = InputWidth * Scale;
 	const int32 Height = InputHeight * Scale;
diff --git a/Engine/Plugins/Runtime/AR/AppleAR/AppleARKit/Source/AppleARKit/Private/AppleARKitVideoOverlay.cpp b/Engine/Plugins/Runtime/AR/AppleAR/AppleARKit/Source/AppleARKit/Private/AppleARKitVideoOverlay.cpp
--- a/Engine/Plugins/Runtime/AR/AppleAR/AppleARKit/Source/AppleARKit/Private/AppleARKitVideoOverlay.cpp
+++ b/Engine/Plugins/Runtime/AR/AppleAR/AppleARKit/Source/AppleARKit/Private/AppleARKitVideoOverlay.cpp
@@ -354,7 +354,7 @@ void FAppleARKitVideoOverlay::RenderVideoOverlayWithMaterial(FRHICommandList& RH
 		FMemory::Memcpy(IndexBuffer.GetData(), Indices, NumIndices * sizeof(uint16));
 
 		FRHIResourceCreateInfo CreateInfoIB(TEXT("VideoOverlayIndexBuffer"), &IndexBuffer);
-		IndexBufferRHI = RHIcmdList.CreateIndexBuffer(sizeof(uint16), IndexBuffer.GetResourceDataSize(), BUF_Static, CreateInfoIB);
+		IndexBufferRHI = RHICmdList.CreateIndexBuffer(sizeof(uint16), IndexBuffer.GetResourceDataSize(), BUF_Static, CreateInfoIB);
 	}
 
 	const auto FeatureLevel = InView.GetFeatureLevel();
