Fixed build break.

diff --git a/Engine/Plugins/Experimental/GPULightmass/Source/GPULightmass/Private/Scene/Scene.cpp b/Engine/Plugins/Experimental/GPULightmass/Source/GPULightmass/Private/Scene/Scene.cpp
--- a/Engine/Plugins/Experimental/GPULightmass/Source/GPULightmass/Private/Scene/Scene.cpp
+++ b/Engine/Plugins/Experimental/GPULightmass/Source/GPULightmass/Private/Scene/Scene.cpp
@@ -630,7 +630,7 @@ void FScene::AddLight(USkyLightComponent* SkyLight)
 		NewSkyLightRenderState.ProcessedTexture = ProcessedSkyTexture->TextureRHI;
 		NewSkyLightRenderState.ProcessedTextureSampler = ProcessedSkyTexture->SamplerStateRHI;
 
-		NewSkyLightRenderState.SkyIrradianceEnvironmentMap.Initialize(TEXT("SkyIrradianceEnvironmentMap"), sizeof(FVector4f), 7);
+		NewSkyLightRenderState.SkyIrradianceEnvironmentMap.Initialize(RHICmdList, TEXT("SkyIrradianceEnvironmentMap"), sizeof(FVector4f), 7);
 
 		NewSkyLightRenderState.PrepareSkyTexture(RHICmdList, FeatureLevel);
 
diff --git a/Engine/Plugins/Experimental/GPULightmass/Source/GPULightmass/Private/VolumetricLightmap.cpp b/Engine/Plugins/Experimental/GPULightmass/Source/GPULightmass/Private/VolumetricLightmap.cpp
--- a/Engine/Plugins/Experimental/GPULightmass/Source/GPULightmass/Private/VolumetricLightmap.cpp
+++ b/Engine/Plugins/Experimental/GPULightmass/Source/GPULightmass/Private/VolumetricLightmap.cpp
@@ -33,7 +33,7 @@ IMPLEMENT_GLOBAL_SHADER(FPermuteVoxelizeVolumeCS, "/Plugin/GPULightmass/Private/
 
 IMPLEMENT_STATIC_UNIFORM_BUFFER_STRUCT(FVLMVoxelizationParams, "VLMVoxelizationParams", SceneTextures);
 
-void InitializeBrickData(FIntVector BrickDataDimensions, FVolumetricLightmapBrickData& BrickData, const bool bForAccumulation)
+void InitializeBrickData(FRHICommandListBase& RHICmdList, FIntVector BrickDataDimensions, FVolumetricLightmapBrickData& BrickData, const bool bForAccumulation)
 {
 	BrickData.AmbientVector.Format = bForAccumulation ? PF_A32B32G32R32F : PF_FloatR11G11B10;
 	BrickData.SkyBentNormal.Format = bForAccumulation ? PF_A32B32G32R32F : PF_R8G8B8A8;
@@ -45,19 +45,19 @@ void InitializeBrickData(FIntVector BrickDataDimensions, FVolumetricLightmapBric
 	}
 
 	BrickData.AmbientVector.CreateTargetTexture(BrickDataDimensions);
-	BrickData.AmbientVector.CreateUAV();
+	BrickData.AmbientVector.CreateUAV(RHICmdList);
 
 	for (int32 i = 0; i < UE_ARRAY_COUNT(BrickData.SHCoefficients); i++)
 	{
 		BrickData.SHCoefficients[i].CreateTargetTexture(BrickDataDimensions);
-		BrickData.SHCoefficients[i].CreateUAV();
+		BrickData.SHCoefficients[i].CreateUAV(RHICmdList);
 	}
 
 	BrickData.SkyBentNormal.CreateTargetTexture(BrickDataDimensions);
-	BrickData.SkyBentNormal.CreateUAV();
+	BrickData.SkyBentNormal.CreateUAV(RHICmdList);
 
 	BrickData.DirectionalLightShadowing.CreateTargetTexture(BrickDataDimensions);
-	BrickData.DirectionalLightShadowing.CreateUAV();
+	BrickData.DirectionalLightShadowing.CreateUAV(RHICmdList);
 
 	size_t TotalSize = 0;
 	TotalSize += GPixelFormats[BrickData.AmbientVector.Format].BlockBytes;
@@ -395,7 +395,7 @@ void FVolumetricLightmapRenderer::VoxelizeScene()
 			TResourceArray<int32> InitialBrickAllocatorParams;
 			InitialBrickAllocatorParams.Add(0);
 			InitialBrickAllocatorParams.Add(0);
-			BrickAllocatorParameters.Initialize(TEXT("VolumetricLightmapBrickAllocatorParameters"), 4, 2, PF_R32_SINT, BUF_UnorderedAccess | BUF_SourceCopy, &InitialBrickAllocatorParams);
+			BrickAllocatorParameters.Initialize(RHICmdList, TEXT("VolumetricLightmapBrickAllocatorParameters"), 4, 2, PF_R32_SINT, BUF_UnorderedAccess | BUF_SourceCopy, &InitialBrickAllocatorParams);
 
 			RHICmdList.Transition(FRHITransitionInfo(BrickAllocatorParameters.UAV, ERHIAccess::Unknown, ERHIAccess::UAVCompute));
 		}
@@ -448,7 +448,7 @@ void FVolumetricLightmapRenderer::VoxelizeScene()
 		BrickLayoutDimensions.Z = FMath::Min(BrickTextureLinearAllocator, MaxBricksInLayoutOneDim);
 	}
 
-	InitializeBrickData(BrickLayoutDimensions * 5, VolumetricLightmapData.BrickData, false);
+	InitializeBrickData(RHICmdList, BrickLayoutDimensions * 5, VolumetricLightmapData.BrickData, false);
 
 	FIntVector BrickLayoutDimensionsForAccumulation;
 
@@ -461,9 +461,9 @@ void FVolumetricLightmapRenderer::VoxelizeScene()
 		BrickLayoutDimensionsForAccumulation.Z = FMath::Min(BrickTextureLinearAllocator, MaxBricksInLayoutOneDim);
 	}
 	
-	InitializeBrickData(BrickLayoutDimensionsForAccumulation * 5, AccumulationBrickData, true);
+	InitializeBrickData(RHICmdList, BrickLayoutDimensionsForAccumulation * 5, AccumulationBrickData, true);
 	
-	BrickRequests.Initialize(TEXT("BrickRequests"), 16, NumTotalBricks, PF_R32G32B32A32_UINT, BUF_UnorderedAccess);
+	BrickRequests.Initialize(RHICmdList, TEXT("BrickRequests"), 16, NumTotalBricks, PF_R32G32B32A32_UINT, BUF_UnorderedAccess);
 
 	VolumetricLightmapData.BrickDataDimensions = BrickLayoutDimensions * 5;
 
