GPULM: Fix BentNormal not being saved.

diff --git a/Engine/Plugins/Experimental/GPULightmass/Source/GPULightmass/Private/Scene/Scene.cpp b/Engine/Plugins/Experimental/GPULightmass/Source/GPULightmass/Private/Scene/Scene.cpp
--- a/Engine/Plugins/Experimental/GPULightmass/Source/GPULightmass/Private/Scene/Scene.cpp
+++ b/Engine/Plugins/Experimental/GPULightmass/Source/GPULightmass/Private/Scene/Scene.cpp
@@ -1952,6 +1952,7 @@ void FScene::ApplyFinishedLightmapsToWorld()
 					ReadbackVolumetricLightmapDataLayerFromGPU(RHICmdList, SubLevelData.BrickData.SHCoefficients[i], SubLevelData.BrickDataDimensions);
 				}
 				ReadbackVolumetricLightmapDataLayerFromGPU(RHICmdList, SubLevelData.BrickData.DirectionalLightShadowing, SubLevelData.BrickDataDimensions);
+				ReadbackVolumetricLightmapDataLayerFromGPU(RHICmdList, SubLevelData.BrickData.SkyBentNormal, SubLevelData.BrickDataDimensions);
 			});
 		}
 
