Add back a missing virtual method implementation

diff --git a/Engine/Plugins/Experimental/GPULightmass/Source/GPULightmass/Private/LightmapPreviewVirtualTexture.h b/Engine/Plugins/Experimental/GPULightmass/Source/GPULightmass/Private/LightmapPreviewVirtualTexture.h
--- a/Engine/Plugins/Experimental/GPULightmass/Source/GPULightmass/Private/LightmapPreviewVirtualTexture.h
+++ b/Engine/Plugins/Experimental/GPULightmass/Source/GPULightmass/Private/LightmapPreviewVirtualTexture.h
@@ -14,6 +14,11 @@ class FLightmapPreviewVirtualTexture : public IVirtualTexture
 public:
 	FLightmapPreviewVirtualTexture(FLightmapRenderStateRef LightmapRenderState, FLightmapRenderer* Renderer);
 
+	virtual bool IsPageStreamed(uint8 vLevel, uint32 vAddress) const override
+	{
+		return false;
+	}
+
 	virtual FVTRequestPageResult RequestPageData(
 		FRHICommandList& RHICmdList,
 		const FVirtualTextureProducerHandle& ProducerHandle,
