GPULM: Hide decal segments using hidden material when baking lightmaps.

diff --git a/Engine/Plugins/Experimental/GPULightmass/Source/GPULightmass/Private/LightmapRenderer.cpp b/Engine/Plugins/Experimental/GPULightmass/Source/GPULightmass/Private/LightmapRenderer.cpp
--- a/Engine/Plugins/Experimental/GPULightmass/Source/GPULightmass/Private/LightmapRenderer.cpp
+++ b/Engine/Plugins/Experimental/GPULightmass/Source/GPULightmass/Private/LightmapRenderer.cpp
@@ -1276,6 +1276,9 @@ bool FSceneRenderState::SetupRayTracingScene(FRDGBuilder& GraphBuilder, int32 LO
 			TArray<FRHIRayTracingShader*> RayTracingHitGroupLibrary;
 			FShaderMapResource::GetRayTracingHitGroupLibrary(RayTracingHitGroupLibrary, DefaultClosestHitShader);
 
+			FRHIRayTracingShader* HiddenMaterialShader = GetGPULightmassDefaultHiddenHitShader(View.ShaderMap);
+			RayTracingHitGroupLibrary.Add(HiddenMaterialShader);
+
 			PSOInitializer.SetHitGroupTable(RayTracingHitGroupLibrary);
 
 			FRHIRayTracingShader* MissTable[] = { GetGPULightmassDefaultMissShader(GlobalShaderMap) };
@@ -1283,6 +1286,8 @@ bool FSceneRenderState::SetupRayTracingScene(FRDGBuilder& GraphBuilder, int32 LO
 
 			RayTracingPipelineState = PipelineStateCache::GetAndOrCreateRayTracingPipelineState(RHICmdList, PSOInitializer);
 
+			const int32 HiddenMaterialIndex = FindRayTracingHitGroupIndex(RayTracingPipelineState, HiddenMaterialShader, true);
+
 			TUniquePtr<FRayTracingLocalShaderBindingWriter> BindingWriter = MakeUnique<FRayTracingLocalShaderBindingWriter>();
 
 			{
@@ -1294,13 +1299,15 @@ bool FSceneRenderState::SetupRayTracingScene(FRDGBuilder& GraphBuilder, int32 LO
 					{
 						const FRayTracingMeshCommand& MeshCommand = *VisibleMeshCommand.RayTracingMeshCommand;
 
+						const int32 MaterialShaderIndex = MeshCommand.bDecal ? HiddenMaterialIndex : MeshCommand.MaterialShaderIndex;
+
 						MeshCommand.SetRayTracingShaderBindingsForHitGroup(BindingWriter.Get(),
 							View.ViewUniformBuffer,
 							SceneUniforms.GetBufferRHI(GraphBuilder),
 							nullptr,
 							VisibleMeshCommand.InstanceIndex,
 							MeshCommand.GeometrySegmentIndex,
-							MeshCommand.MaterialShaderIndex,
+							MaterialShaderIndex,
 							RAY_TRACING_SHADER_SLOT_MATERIAL);
 
 						MeshCommand.SetRayTracingShaderBindingsForHitGroup(BindingWriter.Get(),
@@ -1309,7 +1316,7 @@ bool FSceneRenderState::SetupRayTracingScene(FRDGBuilder& GraphBuilder, int32 LO
 							nullptr,
 							VisibleMeshCommand.InstanceIndex,
 							MeshCommand.GeometrySegmentIndex,
-							MeshCommand.MaterialShaderIndex,
+							MaterialShaderIndex,
 							RAY_TRACING_SHADER_SLOT_SHADOW);
 					}
 				}
@@ -1318,13 +1325,15 @@ bool FSceneRenderState::SetupRayTracingScene(FRDGBuilder& GraphBuilder, int32 LO
 				{
 					const FRayTracingMeshCommand& MeshCommand = *VisibleMeshCommand.RayTracingMeshCommand;
 
+					const int32 MaterialShaderIndex = MeshCommand.bDecal ? HiddenMaterialIndex : MeshCommand.MaterialShaderIndex;
+
 					MeshCommand.SetRayTracingShaderBindingsForHitGroup(BindingWriter.Get(),
 						View.ViewUniformBuffer,
 						SceneUniforms.GetBufferRHI(GraphBuilder),
 						nullptr,
 						VisibleMeshCommand.InstanceIndex,
 						MeshCommand.GeometrySegmentIndex,
-						MeshCommand.MaterialShaderIndex,
+						MaterialShaderIndex,
 						RAY_TRACING_SHADER_SLOT_MATERIAL);
 
 					MeshCommand.SetRayTracingShaderBindingsForHitGroup(BindingWriter.Get(),
@@ -1333,7 +1342,7 @@ bool FSceneRenderState::SetupRayTracingScene(FRDGBuilder& GraphBuilder, int32 LO
 						nullptr,
 						VisibleMeshCommand.InstanceIndex,
 						MeshCommand.GeometrySegmentIndex,
-						MeshCommand.MaterialShaderIndex,
+						MaterialShaderIndex,
 						RAY_TRACING_SHADER_SLOT_SHADOW);
 				}
 
