GPULightmass - Fix logic for determining if viewport is realtime

diff --git a/Engine/Plugins/Experimental/GPULightmass/Source/GPULightmass/Private/GPULightmassModule.cpp b/Engine/Plugins/Experimental/GPULightmass/Source/GPULightmass/Private/GPULightmassModule.cpp
--- a/Engine/Plugins/Experimental/GPULightmass/Source/GPULightmass/Private/GPULightmassModule.cpp
+++ b/Engine/Plugins/Experimental/GPULightmass/Source/GPULightmass/Private/GPULightmassModule.cpp
@@ -7,6 +7,7 @@
 #include "Interfaces/IPluginManager.h"
 #include "ShaderCore.h"
 #include "GPULightmass.h"
+#include "UnrealEdGlobals.h"
 
 #define LOCTEXT_NAMESPACE "StaticLightingSystem"
 
@@ -105,4 +106,9 @@ bool FGPULightmassModule::IsStaticLightingSystemRunning()
 	return StaticLightingSystems.Num() > 0;
 }
 
+bool FGPULightmassModule::IsRealtimeOn()
+{
+	return GCurrentLevelEditingViewportClient && GCurrentLevelEditingViewportClient->IsRealtime();
+}
+
 #undef LOCTEXT_NAMESPACE
\ No newline at end of file
diff --git a/Engine/Plugins/Experimental/GPULightmass/Source/GPULightmass/Private/LightmapRenderer.cpp b/Engine/Plugins/Experimental/GPULightmass/Source/GPULightmass/Private/LightmapRenderer.cpp
--- a/Engine/Plugins/Experimental/GPULightmass/Source/GPULightmass/Private/LightmapRenderer.cpp
+++ b/Engine/Plugins/Experimental/GPULightmass/Source/GPULightmass/Private/LightmapRenderer.cpp
@@ -2018,7 +2018,7 @@ void FLightmapRenderer::Finalize(FRDGBuilder& GraphBuilder)
 		return;
 	}
 
-	bool bIsViewportNonRealtime = GCurrentLevelEditingViewportClient && !GCurrentLevelEditingViewportClient->IsRealtime();
+	const bool bIsViewportNonRealtime = !FGPULightmassModule::IsRealtimeOn();
 
 	int32 MostCommonLODIndex = 0;
 
@@ -3599,7 +3599,7 @@ void FLightmapRenderer::BackgroundTick()
 		}
 	}
 
-	bool bIsViewportNonRealtime = GCurrentLevelEditingViewportClient && !GCurrentLevelEditingViewportClient->IsRealtime();
+	const bool bIsViewportNonRealtime = !FGPULightmassModule::IsRealtimeOn();
 
 	if (bIsViewportNonRealtime && !bWasRunningAtFullSpeed)
 	{
diff --git a/Engine/Plugins/Experimental/GPULightmass/Source/GPULightmass/Private/Scene/Scene.cpp b/Engine/Plugins/Experimental/GPULightmass/Source/GPULightmass/Private/Scene/Scene.cpp
--- a/Engine/Plugins/Experimental/GPULightmass/Source/GPULightmass/Private/Scene/Scene.cpp
+++ b/Engine/Plugins/Experimental/GPULightmass/Source/GPULightmass/Private/Scene/Scene.cpp
@@ -1505,7 +1505,7 @@ void FScene::BackgroundTick()
 	
 	if (GPULightmass->LightBuildNotification.IsValid())
 	{
-		bool bIsViewportNonRealtime = GCurrentLevelEditingViewportClient && !GCurrentLevelEditingViewportClient->IsRealtime();
+		const bool bIsViewportNonRealtime = !FGPULightmassModule::IsRealtimeOn();
 		if (bIsViewportNonRealtime)
 		{
 			if (GPULightmass->Settings->Mode == EGPULightmassMode::FullBake)
@@ -1667,7 +1667,7 @@ void FSceneRenderState::BackgroundTick()
 	LightmapRenderer->BackgroundTick();
 
 	// If we're in background baking mode, schedule VLM work to be after surface lightmaps
-	bool bIsViewportNonRealtime = GCurrentLevelEditingViewportClient && !GCurrentLevelEditingViewportClient->IsRealtime();
+	const bool bIsViewportNonRealtime = !FGPULightmassModule::IsRealtimeOn();
 	if (!bIsViewportNonRealtime || (bIsViewportNonRealtime && bHaveFinishedSurfaceLightmaps))
 	{
 		VolumetricLightmapRenderer->BackgroundTick();
diff --git a/Engine/Plugins/Experimental/GPULightmass/Source/GPULightmass/Private/VolumetricLightmap.cpp b/Engine/Plugins/Experimental/GPULightmass/Source/GPULightmass/Private/VolumetricLightmap.cpp
--- a/Engine/Plugins/Experimental/GPULightmass/Source/GPULightmass/Private/VolumetricLightmap.cpp
+++ b/Engine/Plugins/Experimental/GPULightmass/Source/GPULightmass/Private/VolumetricLightmap.cpp
@@ -752,7 +752,7 @@ void FVolumetricLightmapRenderer::BackgroundTick()
 
 	FGlobalShaderMap* GlobalShaderMap = GetGlobalShaderMap(Scene->FeatureLevel);
 
-	const bool bIsViewportNonRealtime = GCurrentLevelEditingViewportClient && !GCurrentLevelEditingViewportClient->IsRealtime();
+	const bool bIsViewportNonRealtime = !FGPULightmassModule::IsRealtimeOn();
 	const int32 NumSamplesPerFrame = bIsViewportNonRealtime ? Scene->Settings->TilePassesInFullSpeedMode : Scene->Settings->TilePassesInSlowMode;
 
 	FVolumetricLightmapPathTracingRGS::FParameters* PreviousPassParameters = nullptr;
diff --git a/Engine/Plugins/Experimental/GPULightmass/Source/GPULightmass/Public/GPULightmassModule.h b/Engine/Plugins/Experimental/GPULightmass/Source/GPULightmass/Public/GPULightmassModule.h
--- a/Engine/Plugins/Experimental/GPULightmass/Source/GPULightmass/Public/GPULightmassModule.h
+++ b/Engine/Plugins/Experimental/GPULightmass/Source/GPULightmass/Public/GPULightmassModule.h
@@ -23,6 +23,8 @@ public:
 	class FGPULightmass* CreateGPULightmassForWorld(UWorld* InWorld, class UGPULightmassSettings* Settings);
 	void RemoveGPULightmassFromWorld(UWorld* InWorld);
 
+	static bool IsRealtimeOn();
+
 	// Due to limitations in our TMap implementation I cannot use TUniquePtr here
 	// But the GPULightmassModule is the only owner of all static lighting systems, and all worlds weak refer to the systems
 	TMap<UWorld*, class FGPULightmass*> StaticLightingSystems;
diff --git a/Engine/Plugins/Experimental/GPULightmass/Source/GPULightmassEditor/Private/GPULightmassEditorModule.cpp b/Engine/Plugins/Experimental/GPULightmass/Source/GPULightmassEditor/Private/GPULightmassEditorModule.cpp
--- a/Engine/Plugins/Experimental/GPULightmass/Source/GPULightmassEditor/Private/GPULightmassEditorModule.cpp
+++ b/Engine/Plugins/Experimental/GPULightmass/Source/GPULightmassEditor/Private/GPULightmassEditorModule.cpp
@@ -437,9 +437,9 @@ bool FGPULightmassEditorModule::IsBakeWhatYouSeeMode()
 	return false;
 }
 
-bool FGPULightmassEditorModule::IsRealtimeOn() 
+bool FGPULightmassEditorModule::IsRealtimeOn()
 {
-	return GCurrentLevelEditingViewportClient && GCurrentLevelEditingViewportClient->IsRealtime();
+	return FGPULightmassModule::IsRealtimeOn();
 }
 
 bool FGPULightmassEditorModule::IsRunning() 
