GPULightmass - Fix missing header include

diff --git a/Engine/Plugins/Experimental/GPULightmass/Source/GPULightmass/Private/GPULightmassModule.cpp b/Engine/Plugins/Experimental/GPULightmass/Source/GPULightmass/Private/GPULightmassModule.cpp
--- a/Engine/Plugins/Experimental/GPULightmass/Source/GPULightmass/Private/GPULightmassModule.cpp
+++ b/Engine/Plugins/Experimental/GPULightmass/Source/GPULightmass/Private/GPULightmassModule.cpp
@@ -8,6 +8,7 @@
 #include "ShaderCore.h"
 #include "GPULightmass.h"
 #include "UnrealEdGlobals.h"
+#include "LevelEditorViewport.h"
 
 #define LOCTEXT_NAMESPACE "StaticLightingSystem"
 
