GPULightmass - Fix missing GPULIGHTMASS_API

diff --git a/Engine/Plugins/Experimental/GPULightmass/Source/GPULightmass/Public/GPULightmassModule.h b/Engine/Plugins/Experimental/GPULightmass/Source/GPULightmass/Public/GPULightmassModule.h
--- a/Engine/Plugins/Experimental/GPULightmass/Source/GPULightmass/Public/GPULightmassModule.h
+++ b/Engine/Plugins/Experimental/GPULightmass/Source/GPULightmass/Public/GPULightmassModule.h
@@ -23,7 +23,7 @@ public:
 	class FGPULightmass* CreateGPULightmassForWorld(UWorld* InWorld, class UGPULightmassSettings* Settings);
 	void RemoveGPULightmassFromWorld(UWorld* InWorld);
 
-	static bool IsRealtimeOn();
+	static bool GPULIGHTMASS_API IsRealtimeOn();
 
 	// Due to limitations in our TMap implementation I cannot use TUniquePtr here
 	// But the GPULightmassModule is the only owner of all static lighting systems, and all worlds weak refer to the systems
