GPU-LM: Don't cache a scene uniform buffer between runs, as it holds onto references to RDG resources that drop out of scope

diff --git a/Engine/Plugins/Experimental/GPULightmass/Source/GPULightmass/Private/LightmapRenderer.cpp b/Engine/Plugins/Experimental/GPULightmass/Source/GPULightmass/Private/LightmapRenderer.cpp
--- a/Engine/Plugins/Experimental/GPULightmass/Source/GPULightmass/Private/LightmapRenderer.cpp
+++ b/Engine/Plugins/Experimental/GPULightmass/Source/GPULightmass/Private/LightmapRenderer.cpp
@@ -298,7 +298,7 @@ int32 FSceneRenderState::GetPrimitiveIdForGPUScene(const FGeometryInstanceRender
 	return INDEX_NONE;
 }
 
-void FCachedRayTracingSceneData::SetupViewAndSceneUniformBufferFromSceneRenderState(FRDGBuilder& GraphBuilder, FSceneRenderState& Scene)
+void FCachedRayTracingSceneData::SetupViewAndSceneUniformBufferFromSceneRenderState(FRDGBuilder& GraphBuilder, FSceneRenderState& Scene, FSceneUniformBuffer& SceneUniforms)
 {
 	TArray<FPrimitiveSceneShaderData>	PrimitiveSceneData;
 	TArray<FLightmapSceneShaderData>	LightmapSceneData;
@@ -613,7 +613,7 @@ void FCachedRayTracingSceneData::SetupViewAndSceneUniformBufferFromSceneRenderSt
 			GPUScene.GPUSceneLightData = GraphBuilder.CreateSRV(RDGLightDataBuffer);
 		}
 
-		Scene.SceneUniforms.Set(SceneUB::GPUScene, GPUScene);
+		SceneUniforms.Set(SceneUB::GPUScene, GPUScene);
 	}
 }
 
@@ -780,11 +780,11 @@ void FCachedRayTracingSceneData::SetupFromSceneRenderState(FSceneRenderState& Sc
 #endif // RHI_RAYTRACING
 }
 
-void FCachedRayTracingSceneData::RestoreCachedBuffers(FRDGBuilder& GraphBuilder, FSceneRenderState& RenderState)
+void FCachedRayTracingSceneData::RestoreCachedBuffers(FRDGBuilder& GraphBuilder, FSceneUniformBuffer& SceneUniforms)
 {
 	TRACE_CPUPROFILER_EVENT_SCOPE(RestoreCachedBuffers);
 
-	FGPUSceneResourceParameters GPUScene = RenderState.SceneUniforms.Get(SceneUB::GPUScene);
+	FGPUSceneResourceParameters GPUScene = {};
 	
 	check(GPUScenePrimitiveDataBuffer.IsValid())
 	GPUScene.GPUScenePrimitiveSceneData = GraphBuilder.CreateSRV(GraphBuilder.RegisterExternalBuffer(GPUScenePrimitiveDataBuffer));
@@ -802,10 +802,7 @@ void FCachedRayTracingSceneData::RestoreCachedBuffers(FRDGBuilder& GraphBuilder,
 	check(GPUSceneLightDataBuffer.IsValid())
 	GPUScene.GPUSceneLightData = GraphBuilder.CreateSRV(GraphBuilder.RegisterExternalBuffer(GPUSceneLightDataBuffer));
 
-	// Clear cached parameters so that bGPUSceneIsDirty will be true
-	// Required as the uniform buffer in SceneUniforms has only 1 frame lifetime and needs to be recreated
-	RenderState.SceneUniforms.Set(SceneUB::GPUScene, {});
-	RenderState.SceneUniforms.Set(SceneUB::GPUScene, GPUScene);
+	SceneUniforms.Set(SceneUB::GPUScene, GPUScene);
 }
 
 FCachedRayTracingSceneData::~FCachedRayTracingSceneData()
@@ -816,7 +813,7 @@ FCachedRayTracingSceneData::~FCachedRayTracingSceneData()
 	RHICmdList.EnqueueLambda([OwnedRayTracingInstanceTransforms = MoveTemp(OwnedRayTracingInstanceTransforms)](FRHICommandList&){});
 }
 
-bool FSceneRenderState::SetupRayTracingScene(FRDGBuilder& GraphBuilder, int32 LODIndex)
+bool FSceneRenderState::SetupRayTracingScene(FRDGBuilder& GraphBuilder, FSceneUniformBuffer& SceneUniforms, int32 LODIndex)
 {
 	TRACE_CPUPROFILER_EVENT_SCOPE(SetupRayTracingScene);
 
@@ -830,13 +827,13 @@ bool FSceneRenderState::SetupRayTracingScene(FRDGBuilder& GraphBuilder, int32 LO
 
 	if (CachedRayTracingScene.IsValid())
 	{
-		CachedRayTracingScene->RestoreCachedBuffers(GraphBuilder, *this);
+		CachedRayTracingScene->RestoreCachedBuffers(GraphBuilder, SceneUniforms);
 	}
 	else
 	{
 		CachedRayTracingScene = MakeUnique<FCachedRayTracingSceneData>();
 
-		CachedRayTracingScene->SetupViewAndSceneUniformBufferFromSceneRenderState(GraphBuilder, *this);
+		CachedRayTracingScene->SetupViewAndSceneUniformBufferFromSceneRenderState(GraphBuilder, *this, SceneUniforms);
 		CachedRayTracingScene->SetupFromSceneRenderState(*this);
 
 		CalculateDistributionPrefixSumForAllLightmaps();
@@ -2060,7 +2057,9 @@ void FLightmapRenderer::Finalize(FRDGBuilder& GraphBuilder)
 	RectLightAtlas::UpdateAtlasTexture(GraphBuilder, Scene->FeatureLevel);
 	IESAtlas::UpdateAtlasTexture(GraphBuilder, Scene->FeatureLevel);
 
-	if (!Scene->SetupRayTracingScene(GraphBuilder, MostCommonLODIndex))
+	FSceneUniformBuffer SceneUniforms;
+
+	if (!Scene->SetupRayTracingScene(GraphBuilder, SceneUniforms, MostCommonLODIndex))
 	{
 		return;
 	}
@@ -2089,7 +2088,7 @@ void FLightmapRenderer::Finalize(FRDGBuilder& GraphBuilder)
 			TRACE_CPUPROFILER_EVENT_SCOPE(InstanceIdsIdentityBuffer);
 
 			TArray<uint32, SceneRenderingAllocator> InstanceIdsIdentity;
-			for (uint32 Index = 0U; Index < FMath::Max(1U, Scene->SceneUniforms.Get(SceneUB::GPUScene).InstanceDataSOAStride); ++Index)
+			for (uint32 Index = 0U; Index < FMath::Max(1U, SceneUniforms.Get(SceneUB::GPUScene).InstanceDataSOAStride); ++Index)
 			{
 				InstanceIdsIdentity.Add(Index);
 			}
@@ -2270,7 +2269,7 @@ void FLightmapRenderer::Finalize(FRDGBuilder& GraphBuilder)
 
 							auto* PassParameters = GraphBuilder.AllocParameters<FLightmapGBufferPassParameters>();
 							PassParameters->View = Scene->ReferenceView->ViewUniformBuffer;
-							PassParameters->Scene = Scene->SceneUniforms.GetBuffer(GraphBuilder);
+							PassParameters->Scene = SceneUniforms.GetBuffer(GraphBuilder);
 							PassParameters->PassUniformBuffer = PassUniformBuffer;
 							PassParameters->InstanceCulling = InstanceCullingUniformBuffer;
 
@@ -2426,7 +2425,7 @@ void FLightmapRenderer::Finalize(FRDGBuilder& GraphBuilder)
 
 									PassParameters->BatchedTiles = GPUBatchedTileRequests.BatchedTilesSRV;
 									PassParameters->ViewUniformBuffer = Scene->ReferenceView->ViewUniformBuffer;
-									PassParameters->Scene = Scene->SceneUniforms.GetBuffer(GraphBuilder);
+									PassParameters->Scene = SceneUniforms.GetBuffer(GraphBuilder);
 									PassParameters->IrradianceCachingParameters = Scene->IrradianceCache->IrradianceCachingParametersUniformBuffer;
 
 									if (PreviousPassParameters[GPUIndex] == nullptr)
@@ -2715,7 +2714,7 @@ void FLightmapRenderer::Finalize(FRDGBuilder& GraphBuilder)
 					{
 						auto* PassParameters = GraphBuilder.AllocParameters<FLightmapGBufferPassParameters>();
 						PassParameters->View = Scene->ReferenceView->ViewUniformBuffer;
-						PassParameters->Scene = Scene->SceneUniforms.GetBuffer(GraphBuilder);
+						PassParameters->Scene = SceneUniforms.GetBuffer(GraphBuilder);
 						PassParameters->PassUniformBuffer = PassUniformBuffer;
 						PassParameters->InstanceCulling = InstanceCullingUniformBuffer;
 
@@ -2782,7 +2781,7 @@ void FLightmapRenderer::Finalize(FRDGBuilder& GraphBuilder)
 
 						FStationaryLightShadowTracingRGS::FParameters* PassParameters = GraphBuilder.AllocParameters<FStationaryLightShadowTracingRGS::FParameters>();
 						PassParameters->ViewUniformBuffer = Scene->ReferenceView->ViewUniformBuffer;
-						PassParameters->Scene = Scene->SceneUniforms.GetBuffer(GraphBuilder);
+						PassParameters->Scene = SceneUniforms.GetBuffer(GraphBuilder);
 						PassParameters->TLAS = Scene->RayTracingSceneSRV;
 						PassParameters->BatchedTiles = GPUBatchedTileRequests.BatchedTilesSRV;
 						PassParameters->LightTypeArray = HoldReference(GraphBuilder, LightTypeSRV);
diff --git a/Engine/Plugins/Experimental/GPULightmass/Source/GPULightmass/Private/Scene/Lights.cpp b/Engine/Plugins/Experimental/GPULightmass/Source/GPULightmass/Private/Scene/Lights.cpp
--- a/Engine/Plugins/Experimental/GPULightmass/Source/GPULightmass/Private/Scene/Lights.cpp
+++ b/Engine/Plugins/Experimental/GPULightmass/Source/GPULightmass/Private/Scene/Lights.cpp
@@ -386,8 +386,9 @@ void FSkyLightRenderState::PrepareSkyTexture(FRHICommandListImmediate& RHICmdLis
 void FDirectionalLightRenderState::RenderStaticShadowDepthMap(FRHICommandListImmediate& RHICmdList, FSceneRenderState& Scene)
 {
 	FRDGBuilder GraphBuilder(RHICmdList);
+	FSceneUniformBuffer SceneUniforms;
 
-	if (!Scene.SetupRayTracingScene(GraphBuilder))
+	if (!Scene.SetupRayTracingScene(GraphBuilder, SceneUniforms))
 	{
 		return;
 	}
@@ -486,8 +487,9 @@ void FDirectionalLightRenderState::RenderStaticShadowDepthMap(FRHICommandListImm
 void FSpotLightRenderState::RenderStaticShadowDepthMap(FRHICommandListImmediate& RHICmdList, FSceneRenderState& Scene)
 {
 	FRDGBuilder GraphBuilder(RHICmdList);
+	FSceneUniformBuffer SceneUniforms;
 
-	if (!Scene.SetupRayTracingScene(GraphBuilder))
+	if (!Scene.SetupRayTracingScene(GraphBuilder, SceneUniforms))
 	{
 		return;
 	}
@@ -605,8 +607,9 @@ void FSpotLightRenderState::RenderStaticShadowDepthMap(FRHICommandListImmediate&
 void FPointLightRenderState::RenderStaticShadowDepthMap(FRHICommandListImmediate& RHICmdList, FSceneRenderState& Scene)
 {
 	FRDGBuilder GraphBuilder(RHICmdList);
+	FSceneUniformBuffer SceneUniforms;
 
-	if (!Scene.SetupRayTracingScene(GraphBuilder))
+	if (!Scene.SetupRayTracingScene(GraphBuilder, SceneUniforms))
 	{
 		return;
 	}
@@ -698,8 +701,9 @@ void FPointLightRenderState::RenderStaticShadowDepthMap(FRHICommandListImmediate
 void FRectLightRenderState::RenderStaticShadowDepthMap(FRHICommandListImmediate& RHICmdList, FSceneRenderState& Scene)
 {
 	FRDGBuilder GraphBuilder(RHICmdList);
+	FSceneUniformBuffer SceneUniforms;
 
-	if (!Scene.SetupRayTracingScene(GraphBuilder))
+	if (!Scene.SetupRayTracingScene(GraphBuilder, SceneUniforms))
 	{
 		return;
 	}
diff --git a/Engine/Plugins/Experimental/GPULightmass/Source/GPULightmass/Private/Scene/Scene.h b/Engine/Plugins/Experimental/GPULightmass/Source/GPULightmass/Private/Scene/Scene.h
--- a/Engine/Plugins/Experimental/GPULightmass/Source/GPULightmass/Private/Scene/Scene.h
+++ b/Engine/Plugins/Experimental/GPULightmass/Source/GPULightmass/Private/Scene/Scene.h
@@ -76,9 +76,9 @@ struct FCachedRayTracingSceneData
 	TRefCountPtr<FRDGPooledBuffer> GPUSceneInstancePayloadDataBuffer;
 	TRefCountPtr<FRDGPooledBuffer> GPUSceneLightDataBuffer;
 
-	void SetupViewAndSceneUniformBufferFromSceneRenderState(FRDGBuilder& GraphBuilder, class FSceneRenderState& Scene);
+	void SetupViewAndSceneUniformBufferFromSceneRenderState(FRDGBuilder& GraphBuilder, class FSceneRenderState& Scene, FSceneUniformBuffer& SceneUniforms);
 	void SetupFromSceneRenderState(class FSceneRenderState& Scene);
-	void RestoreCachedBuffers(FRDGBuilder& GraphBuilder, class FSceneRenderState& Scene);
+	void RestoreCachedBuffers(FRDGBuilder& GraphBuilder, FSceneUniformBuffer& SceneUniforms);
 };
 
 class FSceneRenderState
@@ -89,8 +89,6 @@ public:
 	void RenderThreadInit();
 	void BackgroundTick();
 
-	FSceneUniformBuffer SceneUniforms;
-
 	FRayTracingSceneRHIRef RayTracingScene;
 	FShaderResourceViewRHIRef RayTracingSceneSRV;
 	FBufferRHIRef RayTracingSceneBuffer;
@@ -120,7 +118,7 @@ public:
 
 	int32 GetPrimitiveIdForGPUScene(const FGeometryInstanceRenderStateRef& GeometryInstanceRef) const;
 
-	bool SetupRayTracingScene(FRDGBuilder& GraphBuilder, int32 LODIndex = INDEX_NONE);
+	bool SetupRayTracingScene(FRDGBuilder& GraphBuilder, FSceneUniformBuffer& SceneUniforms, int32 LODIndex = INDEX_NONE);
 	void DestroyRayTracingScene();
 
 	void CalculateDistributionPrefixSumForAllLightmaps();
diff --git a/Engine/Plugins/Experimental/GPULightmass/Source/GPULightmass/Private/VolumetricLightmap.cpp b/Engine/Plugins/Experimental/GPULightmass/Source/GPULightmass/Private/VolumetricLightmap.cpp
--- a/Engine/Plugins/Experimental/GPULightmass/Source/GPULightmass/Private/VolumetricLightmap.cpp
+++ b/Engine/Plugins/Experimental/GPULightmass/Source/GPULightmass/Private/VolumetricLightmap.cpp
@@ -263,14 +263,15 @@ void FVolumetricLightmapRenderer::VoxelizeScene()
 		}
 
 		// Setup ray tracing scene with LOD 0
-		if (!Scene->SetupRayTracingScene(GraphBuilder))
+		FSceneUniformBuffer SceneUniforms {};
+		if (!Scene->SetupRayTracingScene(GraphBuilder, SceneUniforms))
 		{
 			return;
 		}
 
 		auto* PassParameters = GraphBuilder.AllocParameters<FVoxelizeMeshPassParameters>();
 		PassParameters->View = Scene->ReferenceView->ViewUniformBuffer;
-		PassParameters->Scene = Scene->SceneUniforms.GetBuffer(GraphBuilder);
+		PassParameters->Scene = SceneUniforms.GetBuffer(GraphBuilder);
 		PassParameters->PassUniformBuffer = GraphBuilder.CreateUniformBuffer(VLMVoxelizationParams);
 		PassParameters->InstanceCulling = FInstanceCullingContext::CreateDummyInstanceCullingUniformBuffer(GraphBuilder);
 
@@ -730,10 +731,11 @@ void FVolumetricLightmapRenderer::BackgroundTick()
 	TRACE_CPUPROFILER_EVENT_SCOPE(FVolumetricLightmapRenderer::BackgroundTick);
 
 	FRDGBuilder GraphBuilder(FRHICommandListExecutor::GetImmediateCommandList());
-
+	FSceneUniformBuffer SceneUniforms {};
+	
 	if (IsRayTracingEnabled())
 	{
-		if (!Scene->SetupRayTracingScene(GraphBuilder))
+		if (!Scene->SetupRayTracingScene(GraphBuilder, SceneUniforms))
 		{
 			return;
 		}
