GPULM: Fix missing emissive - it was not included in the 1st bounce.

diff --git a/Engine/Plugins/Experimental/GPULightmass/Shaders/Private/LightmapPathTracing.usf b/Engine/Plugins/Experimental/GPULightmass/Shaders/Private/LightmapPathTracing.usf
--- a/Engine/Plugins/Experimental/GPULightmass/Shaders/Private/LightmapPathTracing.usf
+++ b/Engine/Plugins/Experimental/GPULightmass/Shaders/Private/LightmapPathTracing.usf
@@ -386,8 +386,7 @@ void PathTracingKernel(
 	{		
 		const bool bIsCameraRay = Bounce == 0;
 		const bool bIsLastBounce = Bounce == MaxBounces;
-		const bool bIncludeEmissive = (Bounce > 1) &&
-									  (EnableEmissive       != 0 || bIsCameraRay);
+		const bool bIncludeEmissive = EnableEmissive;
 		FPathTracingPayload Payload = (FPathTracingPayload)0;
 		
 		if (bIsCameraRay)
