Fixed: Compile errors in GPULightmass

diff --git a/Engine/Plugins/Experimental/GPULightmass/Source/GPULightmass/Private/Scene/Scene.cpp b/Engine/Plugins/Experimental/GPULightmass/Source/GPULightmass/Private/Scene/Scene.cpp
--- a/Engine/Plugins/Experimental/GPULightmass/Source/GPULightmass/Private/Scene/Scene.cpp
+++ b/Engine/Plugins/Experimental/GPULightmass/Source/GPULightmass/Private/Scene/Scene.cpp
@@ -1865,7 +1865,7 @@ void FScene::ApplyFinishedLightmapsToWorld()
 	UWorld* World = GPULightmass->World;
 
 	static const auto CVar = IConsoleManager::Get().FindTConsoleVariableDataInt(TEXT("r.VirtualTexturedLightmaps"));
-	const bool bUseVirtualTextures = (CVar->GetValueOnAnyThread() != 0) && UseVirtualTexturing(World->GetFeatureLevel());
+	const bool bUseVirtualTextures = (CVar->GetValueOnAnyThread() != 0) && UseVirtualTexturing(GetFeatureLevelShaderPlatform(World->GetFeatureLevel()));
 
 	bool bHasSkyShadowing = LightScene.SkyLight.IsSet() && LightScene.SkyLight->CastsStationaryShadow();
 	
