GPULM: Fix IES profile - the atlas index was queried before the IES atlas texture was uploaded so we always got INDEX_NONE.

diff --git a/Engine/Plugins/Experimental/GPULightmass/Source/GPULightmass/Private/PathTracingLightParameters.inl b/Engine/Plugins/Experimental/GPULightmass/Source/GPULightmass/Private/PathTracingLightParameters.inl
--- a/Engine/Plugins/Experimental/GPULightmass/Source/GPULightmass/Private/PathTracingLightParameters.inl
+++ b/Engine/Plugins/Experimental/GPULightmass/Source/GPULightmass/Private/PathTracingLightParameters.inl
@@ -86,7 +86,7 @@ void SetupPathTracingLightParameters(
 		DestLight.Dimensions = FVector2f(Light.SourceRadius, Light.SourceLength);
 		DestLight.Attenuation = 1.0f / Light.AttenuationRadius;
 		DestLight.FalloffExponent = Light.FalloffExponent;
-		DestLight.IESAtlasIndex = Light.IESAtlasIndex;
+		DestLight.IESAtlasIndex = Light.GetIESAtlasIndex();
 
 		DestLight.Flags = PATHTRACER_FLAG_TRANSMISSION_MASK;
 		DestLight.Flags |= PATHTRACER_FLAG_LIGHTING_CHANNEL_MASK;
@@ -116,7 +116,7 @@ void SetupPathTracingLightParameters(
 		DestLight.Shaping = EncodeToF16x2(FVector2f(Light.SpotAngles));
 		DestLight.Attenuation = 1.0f / Light.AttenuationRadius;
 		DestLight.FalloffExponent = Light.FalloffExponent;
-		DestLight.IESAtlasIndex = Light.IESAtlasIndex;
+		DestLight.IESAtlasIndex = Light.GetIESAtlasIndex();
 
 		DestLight.Flags = PATHTRACER_FLAG_TRANSMISSION_MASK;
 		DestLight.Flags |= PATHTRACER_FLAG_LIGHTING_CHANNEL_MASK;
@@ -163,7 +163,7 @@ void SetupPathTracingLightParameters(
 		DestLight.Dimensions = FVector2f(Light.SourceWidth, Light.SourceHeight);
 		DestLight.Attenuation = 1.0f / Light.AttenuationRadius;
 		DestLight.Shaping = EncodeToF16x2(FVector2f(FMath::Cos(FMath::DegreesToRadians(Light.BarnDoorAngle)), Light.BarnDoorLength));
-		DestLight.IESAtlasIndex = Light.IESAtlasIndex;
+		DestLight.IESAtlasIndex = Light.GetIESAtlasIndex();
 
 		DestLight.Flags = PATHTRACER_FLAG_TRANSMISSION_MASK;
 		DestLight.Flags |= PATHTRACER_FLAG_LIGHTING_CHANNEL_MASK;
diff --git a/Engine/Plugins/Experimental/GPULightmass/Source/GPULightmass/Private/Scene/Lights.cpp b/Engine/Plugins/Experimental/GPULightmass/Source/GPULightmass/Private/Scene/Lights.cpp
--- a/Engine/Plugins/Experimental/GPULightmass/Source/GPULightmass/Private/Scene/Lights.cpp
+++ b/Engine/Plugins/Experimental/GPULightmass/Source/GPULightmass/Private/Scene/Lights.cpp
@@ -71,10 +71,18 @@ void FLocalLightRenderState::RenderThreadInit()
 	if (IESTexture)
 	{
 		IESAtlasId = IESAtlas::AddTexture(IESTexture);
-		IESAtlasIndex = IESAtlas::GetAtlasSlot(IESAtlasId);
 	}
 }
 
+float FLocalLightRenderState::GetIESAtlasIndex() const
+{
+	if (IESAtlasId != INDEX_NONE)
+	{
+		return IESAtlas::GetAtlasSlot(IESAtlasId);
+	}
+	return float(INDEX_NONE);
+}
+
 void FLocalLightRenderState::RenderThreadFinalize()
 {
 	if (IESTexture)
diff --git a/Engine/Plugins/Experimental/GPULightmass/Source/GPULightmass/Private/Scene/Lights.h b/Engine/Plugins/Experimental/GPULightmass/Source/GPULightmass/Private/Scene/Lights.h
--- a/Engine/Plugins/Experimental/GPULightmass/Source/GPULightmass/Private/Scene/Lights.h
+++ b/Engine/Plugins/Experimental/GPULightmass/Source/GPULightmass/Private/Scene/Lights.h
@@ -94,9 +94,10 @@ struct FLocalLightRenderState
 	bool bCastShadow = true;
 	int ShadowMapChannel = INDEX_NONE;
 	uint32 IESAtlasId = INDEX_NONE;
-	float IESAtlasIndex = INDEX_NONE;
 	UTextureLightProfile* IESTexture = nullptr;
 	
+	float GetIESAtlasIndex() const;
+
 	TSharedPtr<FLightComponentMapBuildData, ESPMode::ThreadSafe> LightComponentMapBuildData;
 
 	virtual FLightRenderParameters GetLightShaderParameters() const = 0;
