GPULM: Initialize VertexFactory correctly for InstancedStaticMesh otherwise we will crash.

diff --git a/Engine/Plugins/Experimental/GPULightmass/Source/GPULightmass/Private/Scene/Scene.cpp b/Engine/Plugins/Experimental/GPULightmass/Source/GPULightmass/Private/Scene/Scene.cpp
--- a/Engine/Plugins/Experimental/GPULightmass/Source/GPULightmass/Private/Scene/Scene.cpp
+++ b/Engine/Plugins/Experimental/GPULightmass/Source/GPULightmass/Private/Scene/Scene.cpp
@@ -1044,8 +1044,9 @@ void FScene::AddGeometryInstanceFromComponent(UInstancedStaticMeshComponent* InC
 			RelevantPointLightsToAddOnRenderThread,
 			RelevantSpotLightsToAddOnRenderThread,
 			RelevantRectLightsToAddOnRenderThread
-		](FRHICommandListImmediate&) mutable
+		](FRHICommandListImmediate& RHICmdList) mutable
 	{
+		InstanceRenderState.InstancedRenderData->BindBuffersToVertexFactories(RHICmdList);
 
 		FInstanceGroupRenderStateRef InstanceRenderStateRef = RenderState.InstanceGroupRenderStates.Emplace(MoveTemp(InstanceRenderState));
 
