GPULM: Fix InstancedStaticMesh not working
* Fill out PrimitiveIdsBuffer correctly
* Updated instances data when the light bake is started/finished

diff --git a/Engine/Plugins/Experimental/GPULightmass/Source/GPULightmass/Private/Scene/Scene.cpp b/Engine/Plugins/Experimental/GPULightmass/Source/GPULightmass/Private/Scene/Scene.cpp
--- a/Engine/Plugins/Experimental/GPULightmass/Source/GPULightmass/Private/Scene/Scene.cpp
+++ b/Engine/Plugins/Experimental/GPULightmass/Source/GPULightmass/Private/Scene/Scene.cpp
@@ -1025,6 +1025,8 @@ void FScene::AddGeometryInstanceFromComponent(UInstancedStaticMeshComponent* InC
 	InstanceRenderState.LocalBounds = InComponent->CalcBounds(FTransform::Identity);
 	InstanceRenderState.bCastShadow = InComponent->CastShadow && InComponent->bCastStaticShadow;
 
+	InstanceRenderState.ComponentUObject->SetBakedLightingDataChangedAll();
+
 	for (int32 LODIndex = 0; LODIndex < Instance->LODLightmaps.Num(); LODIndex++)
 	{
 		InstanceRenderState.LODPerInstanceLightmapSize.Add(Instance->LODPerInstanceLightmapSize[LODIndex]);
@@ -1120,6 +1122,8 @@ void FScene::RemoveGeometryInstanceFromComponent(UInstancedStaticMeshComponent*
 		HISMC->BuildTreeIfOutdated(false, true);
 	}
 
+	InComponent->SetBakedLightingDataChangedAll();
+
 	ENQUEUE_RENDER_COMMAND(RenderThreadRemove)(
 		[ElementId, &RenderState = RenderState](FRHICommandListImmediate&) mutable
 	{
diff --git a/Engine/Source/Runtime/Engine/Classes/Components/InstancedStaticMeshComponent.h b/Engine/Source/Runtime/Engine/Classes/Components/InstancedStaticMeshComponent.h
--- a/Engine/Source/Runtime/Engine/Classes/Components/InstancedStaticMeshComponent.h
+++ b/Engine/Source/Runtime/Engine/Classes/Components/InstancedStaticMeshComponent.h
@@ -574,6 +574,11 @@ public:
 	 */
 	ENGINE_API void SetBakedLightingDataChanged(int32 InInstanceIndex);
 
+	/**
+	 * Mark the "shadowmap" or lightmap uv as modified for all the instances since this is stored in external data.
+	 */
+	ENGINE_API void SetBakedLightingDataChangedAll();
+
 	/** 
 	 * Clears all the updated instance tracking data AND instance ID association, also forcing a full update of the instance data the next time it is flushed.	 
 	 * NOTE: Destroying the instance updated tracking means the renderer has to treat the instances as completely new, preventing e.g., velocity tracking and caching from working reliably.
diff --git a/Engine/Source/Runtime/Engine/Private/InstancedStaticMesh.cpp b/Engine/Source/Runtime/Engine/Private/InstancedStaticMesh.cpp
--- a/Engine/Source/Runtime/Engine/Private/InstancedStaticMesh.cpp
+++ b/Engine/Source/Runtime/Engine/Private/InstancedStaticMesh.cpp
@@ -4411,6 +4411,11 @@ void UInstancedStaticMeshComponent::SetBakedLightingDataChanged(int32 InInstance
 	PrimitiveInstanceDataManager.BakedLightingDataChanged(InInstanceIndex);
 }
 
+void UInstancedStaticMeshComponent::SetBakedLightingDataChangedAll()
+{
+	PrimitiveInstanceDataManager.BakedLightingDataChangedAll();
+}
+
 void UInstancedStaticMeshComponent::InvalidateInstanceDataTracking()
 {
 	PrimitiveInstanceDataManager.Invalidate(GetNumInstances());
