Fix minor shader warnings

diff --git a/Engine/Plugins/Experimental/GPULightmass/Shaders/Private/BrickAllocationManagement.usf b/Engine/Plugins/Experimental/GPULightmass/Shaders/Private/BrickAllocationManagement.usf
--- a/Engine/Plugins/Experimental/GPULightmass/Shaders/Private/BrickAllocationManagement.usf
+++ b/Engine/Plugins/Experimental/GPULightmass/Shaders/Private/BrickAllocationManagement.usf
@@ -29,7 +29,7 @@ float3 ImportanceVolumeMax;
 void VoxelizeImportanceVolumeCS(uint3 VoxelPos : SV_DispatchThreadID) 
 {
 	if (any(VoxelPos >= VolumeSize)) return; 
-	float3 WorldPosition = ((float3)VoxelPos / VLMVoxelizationParams.VolumeMaxDim - 0.5f) * 2.0f * VLMVoxelizationParams.VolumeExtent + VLMVoxelizationParams.VolumeCenter;
+	float3 WorldPosition = ((float3)VoxelPos / VLMVoxelizationParams.VolumeMaxDim - 0.5f) * 2.0f * VLMVoxelizationParams.VolumeExtent.xyz + VLMVoxelizationParams.VolumeCenter.xyz;
 	if (all(WorldPosition >= ImportanceVolumeMin) && all(WorldPosition <= ImportanceVolumeMax))
 	{
 		VoxelizeVolume[VoxelPos] = BRICK_IN_IMPORTANCE_VOLUME;
@@ -336,7 +336,7 @@ UNORM float4 ManualTrilinearFilter(RWTexture3D<UNORM float4> Texture, float3 UVW
 	return Value;
 }
 
-float3 ManualTrilinearFilter(RWTexture3D<UNORM float> Texture, float3 UVW)
+float ManualTrilinearFilter(RWTexture3D<UNORM float> Texture, float3 UVW)
 {
 	int3 Coord = UVW;
 	float3 CoordinateFraction = frac(UVW);
@@ -373,7 +373,7 @@ void CopyPaddingStripsCS(uint3 BrickID : SV_GroupID, uint3 CellPosInBrick : SV_G
 	float PaddedBrickSize = 4 + 1;
 	uint3 VoxelPos = ComputeBrickLayoutPosition(BrickIndex, uint3(256, 256, 256)) * PaddedBrickSize + CellPosInBrick;
 
-	if (BrickRequests[BrickIndex].w == 1 << (2 * MipLevel))
+	if (BrickRequests[BrickIndex].w == 1u << (2 * MipLevel))
 	{
 		int3 VoxelizeVolumeCoords = floor(CellPosInCurrentMip * InvCellSizeInCurrentMip);
 		int BrickLinearAddress = VoxelizeVolume[VoxelizeVolumeCoords] - 1;
