GPULM: Fix Landscape vertex factory incorrect position component format

diff --git a/Engine/Plugins/Experimental/GPULightmass/Source/GPULightmass/Private/Scene/Scene.cpp b/Engine/Plugins/Experimental/GPULightmass/Source/GPULightmass/Private/Scene/Scene.cpp
--- a/Engine/Plugins/Experimental/GPULightmass/Source/GPULightmass/Private/Scene/Scene.cpp
+++ b/Engine/Plugins/Experimental/GPULightmass/Source/GPULightmass/Private/Scene/Scene.cpp
@@ -1279,7 +1279,7 @@ void FScene::AddGeometryInstanceFromComponent(ULandscapeComponent* InComponent)
 			FLandscapeComponentSceneProxy::SharedBuffersMap.Add(InstanceRenderState.SharedBuffersKey, InstanceRenderState.SharedBuffers);
 
 			FLandscapeFixedGridVertexFactory* LandscapeVertexFactory = new FLandscapeFixedGridVertexFactory(LocalFeatureLevel);
-			LandscapeVertexFactory->Data.PositionComponent = FVertexStreamComponent(InstanceRenderState.SharedBuffers->VertexBuffer, 0, sizeof(FLandscapeVertex), VET_Float4);
+			LandscapeVertexFactory->Data.PositionComponent = FVertexStreamComponent(InstanceRenderState.SharedBuffers->VertexBuffer, 0, sizeof(FLandscapeVertex), VET_UByte4);
 			LandscapeVertexFactory->InitResource(RHICmdList);
 			InstanceRenderState.SharedBuffers->FixedGridVertexFactory = LandscapeVertexFactory;
 		}
