GPULM: Add cvar to kick off the lighting build from commandline:
* GPULM.BuildLighting - will use currently set parameters

diff --git a/Engine/Plugins/Experimental/GPULightmass/Source/GPULightmassEditor/Private/GPULightmassEditorModule.cpp b/Engine/Plugins/Experimental/GPULightmass/Source/GPULightmassEditor/Private/GPULightmassEditorModule.cpp
--- a/Engine/Plugins/Experimental/GPULightmass/Source/GPULightmassEditor/Private/GPULightmassEditorModule.cpp
+++ b/Engine/Plugins/Experimental/GPULightmass/Source/GPULightmassEditor/Private/GPULightmassEditorModule.cpp
@@ -36,6 +36,23 @@ FName GPULightmassSettingsTabName = TEXT("GPULightmassSettings");
 #include "Windows/AllowWindowsPlatformTypes.h"
 #endif
 
+static FAutoConsoleCommand GPULMBuildLighting(
+	TEXT("GPULM.BuildLighting"),
+	TEXT("Starts lighting build with GPU Lightmass"),
+	FConsoleCommandDelegate::CreateStatic([] {
+		FGPULightmassEditorModule* LightmassModule = FModuleManager::Get().GetModulePtr<FGPULightmassEditorModule>("GPULightmass");
+		if (LightmassModule)
+		{
+			if (UWorld* World = GEditor->GetEditorWorldContext().World())
+			{
+				if (!World->GetSubsystem<UGPULightmassSubsystem>()->IsRunning())
+				{					
+					World->GetSubsystem<UGPULightmassSubsystem>()->Launch();
+				}
+			}
+		}		
+}));
+
 static bool IsRenderDocPresent()
 {
 #if PLATFORM_WINDOWS
