Deprecate public access to RayTracingGeometryRHI in FRayTracingGeometry.
- use GetRHI() instead.

diff --git a/Engine/Plugins/Experimental/GPULightmass/Source/GPULightmass/Private/LightmapRenderer.cpp b/Engine/Plugins/Experimental/GPULightmass/Source/GPULightmass/Private/LightmapRenderer.cpp
--- a/Engine/Plugins/Experimental/GPULightmass/Source/GPULightmass/Private/LightmapRenderer.cpp
+++ b/Engine/Plugins/Experimental/GPULightmass/Source/GPULightmass/Private/LightmapRenderer.cpp
@@ -700,7 +700,7 @@ void FCachedRayTracingSceneData::SetupFromSceneRenderState(FSceneRenderState& Sc
 			{
 				int32 InstanceIndex = RayTracingGeometryInstancesPerLOD[LODIndex].AddDefaulted(1);
 				FRayTracingGeometryInstance& RayTracingInstance = RayTracingGeometryInstancesPerLOD[LODIndex][InstanceIndex];
-				RayTracingInstance.GeometryRHI = Instance.RenderData->LODResources[LODIndexToUse].RayTracingGeometry.RayTracingGeometryRHI;
+				RayTracingInstance.GeometryRHI = Instance.RenderData->LODResources[LODIndexToUse].RayTracingGeometry.GetRHI();
 				RayTracingInstance.Transforms = MakeArrayView(&Instance.LocalToWorld, 1);
 				RayTracingInstance.NumTransforms = 1;
 				RayTracingInstance.DefaultUserData = (uint32)StaticMeshIndex;
@@ -737,7 +737,7 @@ void FCachedRayTracingSceneData::SetupFromSceneRenderState(FSceneRenderState& Sc
 			{
 				int32 InstanceIndex = RayTracingGeometryInstancesPerLOD[LODIndex].AddDefaulted(1);
 				FRayTracingGeometryInstance& RayTracingInstance = RayTracingGeometryInstancesPerLOD[LODIndex][InstanceIndex];
-				RayTracingInstance.GeometryRHI = InstanceGroup.ComponentUObject->GetStaticMesh()->GetRenderData()->LODResources[LODIndexToUse].RayTracingGeometry.RayTracingGeometryRHI;
+				RayTracingInstance.GeometryRHI = InstanceGroup.ComponentUObject->GetStaticMesh()->GetRenderData()->LODResources[LODIndexToUse].RayTracingGeometry.GetRHI();
 
 				const int32 NumInstances = (int32)InstanceGroup.NumInstances;
 				TArrayView<FMatrix> NewTransforms = MakeArrayView(
@@ -1058,7 +1058,7 @@ bool FSceneRenderState::SetupRayTracingScene(FRDGBuilder& GraphBuilder, FSceneUn
 						}
 
 						FRayTracingGeometryInstance& RayTracingInstance = RayTracingGeometryInstances[InstanceIndex];
-						RayTracingInstance.GeometryRHI = Landscape.SectionRayTracingStates[SubSectionIdx]->Geometry.RayTracingGeometryRHI;
+						RayTracingInstance.GeometryRHI = Landscape.SectionRayTracingStates[SubSectionIdx]->Geometry.GetRHI();
 						RayTracingInstance.Transforms = MakeArrayView(LandscapeTransforms.Emplace_GetRef(MakeUnique<FMatrix>(Landscape.LocalToWorld)).Get(), 1);
 						RayTracingInstance.NumTransforms = 1;
 						RayTracingInstance.DefaultUserData = (uint32)(StaticMeshInstanceRenderStates.Elements.Num() + InstanceGroupRenderStates.Elements.Num() + LandscapeIndex);
diff --git a/Engine/Plugins/Runtime/GeometryCache/Source/GeometryCache/Private/GeometryCacheSceneProxy.cpp b/Engine/Plugins/Runtime/GeometryCache/Source/GeometryCache/Private/GeometryCacheSceneProxy.cpp
--- a/Engine/Plugins/Runtime/GeometryCache/Source/GeometryCache/Private/GeometryCacheSceneProxy.cpp
+++ b/Engine/Plugins/Runtime/GeometryCache/Source/GeometryCache/Private/GeometryCacheSceneProxy.cpp
@@ -662,7 +662,7 @@ void FGeometryCacheSceneProxy::UpdateAnimation(FRHICommandListBase& RHICmdList,
 					{
 						// Request full build on same geometry because data might have changed to much for update call?
 						FRayTracingGeometryBuildParams BuildParams;
-						BuildParams.Geometry = Section->RayTracingGeometry.RayTracingGeometryRHI;
+						BuildParams.Geometry = Section->RayTracingGeometry.GetRHI();
 						BuildParams.BuildMode = EAccelerationStructureBuildMode::Build;
 						BuildParams.Segments = Section->RayTracingGeometry.Initializer.Segments;
 						FRHIComputeCommandList::Get(RHICmdList).BuildAccelerationStructures(MakeArrayView(&BuildParams, 1));
diff --git a/Engine/Plugins/Runtime/HairStrands/Source/HairStrandsCore/Private/HairStrandsInterpolation.cpp b/Engine/Plugins/Runtime/HairStrands/Source/HairStrandsCore/Private/HairStrandsInterpolation.cpp
--- a/Engine/Plugins/Runtime/HairStrands/Source/HairStrandsCore/Private/HairStrandsInterpolation.cpp
+++ b/Engine/Plugins/Runtime/HairStrands/Source/HairStrandsCore/Private/HairStrandsInterpolation.cpp
@@ -1330,7 +1330,7 @@ static void UpdateHairAccelerationStructure(FRHICommandList& RHICmdList, FRayTra
 
 	FRayTracingGeometryBuildParams Params;
 	Params.BuildMode = InMode;
-	Params.Geometry = RayTracingGeometry->RayTracingGeometryRHI;
+	Params.Geometry = RayTracingGeometry->GetRHI();
 	Params.Segments = RayTracingGeometry->Initializer.Segments;
 
 	RHICmdList.BuildAccelerationStructures(MakeArrayView(&Params, 1));
diff --git a/Engine/Source/Runtime/Engine/Private/RayTracingSkinnedGeometry.cpp b/Engine/Source/Runtime/Engine/Private/RayTracingSkinnedGeometry.cpp
--- a/Engine/Source/Runtime/Engine/Private/RayTracingSkinnedGeometry.cpp
+++ b/Engine/Source/Runtime/Engine/Private/RayTracingSkinnedGeometry.cpp
@@ -167,7 +167,7 @@ void FRayTracingSkinnedGeometryUpdateQueue::Commit(FRHICommandListImmediate & RH
 			FRayTracingUpdateInfo& UpdateInfo = Iter.Value();
 
 			FRayTracingGeometryBuildParams BuildParams;
-			BuildParams.Geometry = RayTracingGeometry->RayTracingGeometryRHI;
+			BuildParams.Geometry = RayTracingGeometry->GetRHI();
 			BuildParams.BuildMode = UpdateInfo.BuildMode;
 			BuildParams.Segments = RayTracingGeometry->Initializer.Segments;
 
diff --git a/Engine/Source/Runtime/Engine/Private/Rendering/RayTracingGeometryManager.cpp b/Engine/Source/Runtime/Engine/Private/Rendering/RayTracingGeometryManager.cpp
--- a/Engine/Source/Runtime/Engine/Private/Rendering/RayTracingGeometryManager.cpp
+++ b/Engine/Source/Runtime/Engine/Private/Rendering/RayTracingGeometryManager.cpp
@@ -178,7 +178,7 @@ void FRayTracingGeometryManager::Tick(FRHICommandList& RHICmdList)
 			FScopeLock ScopeLock(&RequestCS);
 			for (FRayTracingGeometry* Geometry : RegisteredGeometries)
 			{
-				checkf(Geometry->IsEvicted() || Geometry->RayTracingGeometryRHI == nullptr, TEXT("Ray tracing geometry should be evicted when ray tracing is disabled."));
+				checkf(Geometry->IsEvicted() || Geometry->GetRHI() == nullptr, TEXT("Ray tracing geometry should be evicted when ray tracing is disabled."));
 			}
 		}
 #endif
@@ -205,7 +205,7 @@ void FRayTracingGeometryManager::Tick(FRHICommandList& RHICmdList)
 		FScopeLock ScopeLock(&RequestCS);
 		for (FRayTracingGeometry* Geometry : RegisteredGeometries)
 		{
-			if (Geometry->RayTracingGeometryRHI != nullptr)
+			if (Geometry->GetRHI() != nullptr)
 			{
 				Geometry->Evict();
 			}
@@ -325,7 +325,7 @@ void FRayTracingGeometryManager::SetupBuildParams(const FBuildRequest& InBuildRe
 	check(InBuildRequest.RequestIndex != INDEX_NONE && InBuildRequest.Owner->RayTracingBuildRequestIndex != INDEX_NONE);
 
 	FRayTracingGeometryBuildParams BuildParam;
-	BuildParam.Geometry = InBuildRequest.Owner->RayTracingGeometryRHI;
+	BuildParam.Geometry = InBuildRequest.Owner->GetRHI();
 	BuildParam.BuildMode = InBuildRequest.BuildMode;
 	InBuildParams.Add(BuildParam);
 
diff --git a/Engine/Source/Runtime/Engine/Private/SkeletalRenderCPUSkin.cpp b/Engine/Source/Runtime/Engine/Private/SkeletalRenderCPUSkin.cpp
--- a/Engine/Source/Runtime/Engine/Private/SkeletalRenderCPUSkin.cpp
+++ b/Engine/Source/Runtime/Engine/Private/SkeletalRenderCPUSkin.cpp
@@ -416,7 +416,7 @@ void FSkeletalMeshObjectCPUSkin::FSkeletalMeshObjectLOD::InitResources(FSkelMesh
 					GeometrySections.Add(Segment);
 				}
 				Initializer.Segments = GeometrySections;
-				Initializer.SourceGeometry = SourceGeometry.RayTracingGeometryRHI;
+				Initializer.SourceGeometry = SourceGeometry.GetRHI();
 
 				RayTracingGeometry.SetInitializer(Initializer);
 				RayTracingGeometry.InitResource(RHICmdList);
diff --git a/Engine/Source/Runtime/Engine/Private/SkeletalRenderGPUSkin.cpp b/Engine/Source/Runtime/Engine/Private/SkeletalRenderGPUSkin.cpp
--- a/Engine/Source/Runtime/Engine/Private/SkeletalRenderGPUSkin.cpp
+++ b/Engine/Source/Runtime/Engine/Private/SkeletalRenderGPUSkin.cpp
@@ -884,7 +884,7 @@ void FSkeletalMeshObjectGPUSkin::UpdateRayTracingGeometry(FRHICommandListBase& R
 				Initializer.Segments.Add(Segment);
 			}
 
-			if (RayTracingGeometry.RayTracingGeometryRHI.IsValid())
+			if (RayTracingGeometry.GetRHI() != nullptr)
 			{
 				// RayTracingGeometry.ReleaseRHI() releases the old RT geometry, however due to the deferred deletion nature of RHI resources
 				// they will not be released until the end of the frame. We may get OOM in the middle of batched updates if not flushing.
@@ -892,11 +892,11 @@ void FSkeletalMeshObjectGPUSkin::UpdateRayTracingGeometry(FRHICommandListBase& R
 
 				// Release the old data (make sure it's not pending build anymore either)
 				RayTracingUpdateQueue->Remove(&RayTracingGeometry, MemoryEstimation);
-				RayTracingGeometry.RayTracingGeometryRHI->DisableLifetimeExtension();
+				RayTracingGeometry.GetRHI()->DisableLifetimeExtension();
 				RayTracingGeometry.ReleaseRHI();
 			}
 
-			Initializer.SourceGeometry = LODModel.SourceRayTracingGeometry.RayTracingGeometryRHI;
+			Initializer.SourceGeometry = LODModel.SourceRayTracingGeometry.GetRHI();
 
 			// Get the scratch sizes used for build & update
 			RayTracingGeometryStructureSize = RHICalcRayTracingGeometrySize(Initializer);
diff --git a/Engine/Source/Runtime/RenderCore/Private/RayTracingGeometry.cpp b/Engine/Source/Runtime/RenderCore/Private/RayTracingGeometry.cpp
--- a/Engine/Source/Runtime/RenderCore/Private/RayTracingGeometry.cpp
+++ b/Engine/Source/Runtime/RenderCore/Private/RayTracingGeometry.cpp
@@ -20,6 +20,8 @@ static TAutoConsoleVariable<int32> CVarDebugForceRuntimeBLAS(
 	TEXT("Force building BLAS at runtime."),
 	ECVF_ReadOnly);
 
+PRAGMA_DISABLE_DEPRECATION_WARNINGS // Remove this in UE 5.6 once RayTracingGeometryRHI is made private
+
 FRayTracingGeometry::FRayTracingGeometry() = default;
 
 /** Destructor. */
@@ -295,3 +297,5 @@ void FRayTracingGeometry::BoostBuildPriority(float InBoostValue) const
 }
 
 #endif // RHI_RAYTRACING
+
+PRAGMA_ENABLE_DEPRECATION_WARNINGS
diff --git a/Engine/Source/Runtime/RenderCore/Public/RayTracingGeometry.h b/Engine/Source/Runtime/RenderCore/Public/RayTracingGeometry.h
--- a/Engine/Source/Runtime/RenderCore/Public/RayTracingGeometry.h
+++ b/Engine/Source/Runtime/RenderCore/Public/RayTracingGeometry.h
@@ -47,8 +47,17 @@ public:
 	int64 DynamicGeometrySharedBufferGenerationID = NonSharedVertexBuffers;
 
 	FRayTracingGeometryInitializer Initializer;
+
+	UE_DEPRECATED(5.5, "Use GetRHI() instead.")
 	FRayTracingGeometryRHIRef RayTracingGeometryRHI;
 
+	FRHIRayTracingGeometry* GetRHI() const
+	{
+		PRAGMA_DISABLE_DEPRECATION_WARNINGS
+		return RayTracingGeometryRHI;
+		PRAGMA_ENABLE_DEPRECATION_WARNINGS
+	}
+
 	RayTracing::GeometryGroupHandle GroupHandle = INDEX_NONE;
 
 	/** LOD of the mesh associated with this ray tracing geometry object (-1 if unknown) */
diff --git a/Engine/Source/Runtime/Renderer/Private/PrimitiveSceneInfo.cpp b/Engine/Source/Runtime/Renderer/Private/PrimitiveSceneInfo.cpp
--- a/Engine/Source/Runtime/Renderer/Private/PrimitiveSceneInfo.cpp
+++ b/Engine/Source/Runtime/Renderer/Private/PrimitiveSceneInfo.cpp
@@ -389,7 +389,7 @@ bool FPrimitiveSceneInfo::IsCachedRayTracingGeometryValid() const
 	if (CachedRayTracingGeometry)
 	{
 		// TODO: Doesn't take Nanite Ray Tracing into account
-		check(CachedRayTracingGeometry->RayTracingGeometryRHI == CachedRayTracingInstance.GeometryRHI);
+		check(CachedRayTracingGeometry->GetRHI() == CachedRayTracingInstance.GeometryRHI);
 		check(!CachedRayTracingGeometry->GetRequiresBuild() && !CachedRayTracingGeometry->HasPendingBuildRequest());
 
 		return CachedRayTracingGeometry->IsValid();
@@ -442,7 +442,7 @@ FRHIRayTracingGeometry* FPrimitiveSceneInfo::GetStaticRayTracingGeometryInstance
 		}
 		else if (StaticRayTracingGeometries[LodLevel]->IsValid())
 		{
-			return StaticRayTracingGeometries[LodLevel]->RayTracingGeometryRHI;
+			return StaticRayTracingGeometries[LodLevel]->GetRHI();
 		}
 		else
 		{
@@ -1228,7 +1228,7 @@ void FPrimitiveSceneInfo::UpdateCachedRayTracingInstance(FPrimitiveSceneInfo* Sc
 		{
 			checkf(CachedRayTracingInstance.Geometry, TEXT("Cached ray tracing instances must have valid geometries.")); // unless using nanite ray tracing
 
-			SceneInfo->CachedRayTracingInstance.GeometryRHI = CachedRayTracingInstance.Geometry->RayTracingGeometryRHI;
+			SceneInfo->CachedRayTracingInstance.GeometryRHI = CachedRayTracingInstance.Geometry->GetRHI();
 		}
 
 		// At this point (in AddToScene()) PrimitiveIndex has been set
diff --git a/Engine/Source/Runtime/Renderer/Private/RayTracing/RayTracing.cpp b/Engine/Source/Runtime/Renderer/Private/RayTracing/RayTracing.cpp
--- a/Engine/Source/Runtime/Renderer/Private/RayTracing/RayTracing.cpp
+++ b/Engine/Source/Runtime/Renderer/Private/RayTracing/RayTracing.cpp
@@ -559,7 +559,7 @@ namespace RayTracing
 								RelevantPrimitive->PersistentPrimitiveIndex = SceneInfo->GetPersistentIndex();
 
 								RelevantPrimitive->LODIndex = LODIndex;
-								RelevantPrimitive->RayTracingGeometryRHI = RayTracingGeometry->RayTracingGeometryRHI;
+								RelevantPrimitive->RayTracingGeometryRHI = RayTracingGeometry->GetRHI();
 
 								RelevantPrimitive->CachedRayTracingMeshCommandIndices = SceneInfo->CachedRayTracingMeshCommandIndicesPerLOD[LODIndex];
 								RelevantPrimitive->StateHash = SceneInfo->CachedRayTracingMeshCommandsHashPerLOD[LODIndex];
@@ -919,7 +919,7 @@ namespace RayTracing
 						}
 
 						FRayTracingGeometryInstance RayTracingInstance;
-						RayTracingInstance.GeometryRHI = Geometry->RayTracingGeometryRHI;
+						RayTracingInstance.GeometryRHI = Geometry->GetRHI();
 						checkf(RayTracingInstance.GeometryRHI, TEXT("Ray tracing instance must have a valid geometry."));
 						RayTracingInstance.DefaultUserData = PersistentPrimitiveIndex.Index;
 						RayTracingInstance.bApplyLocalBoundsTransform = Instance.bApplyLocalBoundsTransform;
diff --git a/Engine/Source/Runtime/Renderer/Private/RayTracing/RayTracingDynamicGeometry.cpp b/Engine/Source/Runtime/Renderer/Private/RayTracing/RayTracingDynamicGeometry.cpp
--- a/Engine/Source/Runtime/Renderer/Private/RayTracing/RayTracingDynamicGeometry.cpp
+++ b/Engine/Source/Runtime/Renderer/Private/RayTracing/RayTracingDynamicGeometry.cpp
@@ -351,7 +351,7 @@ void FRayTracingDynamicGeometryCollection::AddDynamicMeshBatchForGeometryUpdate(
 	}
 
 	FRayTracingGeometryBuildParams Params;
-	Params.Geometry = Geometry.RayTracingGeometryRHI;
+	Params.Geometry = Geometry.GetRHI();
 	Params.BuildMode = Geometry.GetRequiresBuild()
 		? EAccelerationStructureBuildMode::Build
 		: EAccelerationStructureBuildMode::Update;
