GPULM: Fix shadowing warnings in lambdas

diff --git a/Engine/Plugins/Experimental/GPULightmass/Source/GPULightmass/Private/GPULightmass.cpp b/Engine/Plugins/Experimental/GPULightmass/Source/GPULightmass/Private/GPULightmass.cpp
--- a/Engine/Plugins/Experimental/GPULightmass/Source/GPULightmass/Private/GPULightmass.cpp
+++ b/Engine/Plugins/Experimental/GPULightmass/Source/GPULightmass/Private/GPULightmass.cpp
@@ -42,18 +42,18 @@ void FGPULightmass::ShowLightBuildNotification()
 			Info.ButtonDetails.Add(FNotificationButtonInfo(
 			LOCTEXT("Save", "Save and Stop"),
 			FText::GetEmpty(),
-			FSimpleDelegate::CreateLambda([World = World, this]() { 
+			FSimpleDelegate::CreateLambda([WorldPtr = World, this]() { 
 				this->Scene.ApplyFinishedLightmapsToWorld(); 
-				World->GetSubsystem<UGPULightmassSubsystem>()->Stop(); 
+				WorldPtr->GetSubsystem<UGPULightmassSubsystem>()->Stop();
 			})));
 		}
 	
 		Info.ButtonDetails.Add(FNotificationButtonInfo(
 			LOCTEXT("LightBuildCancel", "Cancel"),
 			LOCTEXT("LightBuildCancelToolTip", "Cancels the lighting build in progress."),
-			FSimpleDelegate::CreateLambda([World = World]() { 
-				World->GetSubsystem<UGPULightmassSubsystem>()->Stop(); 
-				World->GetSubsystem<UGPULightmassSubsystem>()->OnLightBuildEnded().Broadcast();
+			FSimpleDelegate::CreateLambda([WorldPtr = World]() {
+				WorldPtr->GetSubsystem<UGPULightmassSubsystem>()->Stop();
+				WorldPtr->GetSubsystem<UGPULightmassSubsystem>()->OnLightBuildEnded().Broadcast();
 				})));
 
 		LightBuildNotification = FSlateNotificationManager::Get().AddNotification(Info);
diff --git a/Engine/Plugins/Experimental/GPULightmass/Source/GPULightmass/Private/LightmapRenderer.cpp b/Engine/Plugins/Experimental/GPULightmass/Source/GPULightmass/Private/LightmapRenderer.cpp
--- a/Engine/Plugins/Experimental/GPULightmass/Source/GPULightmass/Private/LightmapRenderer.cpp
+++ b/Engine/Plugins/Experimental/GPULightmass/Source/GPULightmass/Private/LightmapRenderer.cpp
@@ -354,7 +354,7 @@ void FCachedRayTracingSceneData::SetupViewAndSceneUniformBufferFromSceneRenderSt
 				&LightmapSceneData,
 				&LightmapSceneDataStartOffsets,
 				&InstanceSceneData,
-				&InstanceDataOriginalOffsets = InstanceDataOriginalOffsets
+				&InstanceDataOffsets = InstanceDataOriginalOffsets
 			](int32 InstanceIndex = 0)
 			{
 				FStaticMeshInstanceRenderState& Instance = Scene.StaticMeshInstanceRenderStates.Elements[InstanceIndex];
@@ -376,7 +376,7 @@ void FCachedRayTracingSceneData::SetupViewAndSceneUniformBufferFromSceneRenderSt
 				
 				PrimitiveSceneData[InstanceIndex] = FPrimitiveSceneShaderData(PrimitiveUniformShaderParameters);
 
-				InstanceDataOriginalOffsets[InstanceIndex] = InstanceIndex;
+				InstanceDataOffsets[InstanceIndex] = InstanceIndex;
 
 				FInstanceSceneShaderData& SceneData = InstanceSceneData[InstanceIndex];
 				SceneData.Build
