[Slate] Replaced old style editor dialog buttons

diff --git a/Engine/Plugins/2D/Paper2D/Source/Paper2DEditor/Paper2DEditor.Build.cs b/Engine/Plugins/2D/Paper2D/Source/Paper2DEditor/Paper2DEditor.Build.cs
--- a/Engine/Plugins/2D/Paper2D/Source/Paper2DEditor/Paper2DEditor.Build.cs
+++ b/Engine/Plugins/2D/Paper2D/Source/Paper2DEditor/Paper2DEditor.Build.cs
@@ -31,6 +31,7 @@ public class Paper2DEditor : ModuleRules
 				"Projects",
 				"NavigationSystem",
 				"ToolMenus",
+				"ToolWidgets"
             });
 
 		PublicDependencyModuleNames.AddRange(
diff --git a/Engine/Plugins/2D/Paper2D/Source/Paper2DEditor/Private/ExtractSprites/SPaperExtractSpritesDialog.cpp b/Engine/Plugins/2D/Paper2D/Source/Paper2DEditor/Private/ExtractSprites/SPaperExtractSpritesDialog.cpp
--- a/Engine/Plugins/2D/Paper2D/Source/Paper2DEditor/Private/ExtractSprites/SPaperExtractSpritesDialog.cpp
+++ b/Engine/Plugins/2D/Paper2D/Source/Paper2DEditor/Private/ExtractSprites/SPaperExtractSpritesDialog.cpp
@@ -24,6 +24,7 @@
 #include "PaperSpriteFactory.h"
 #include "TextureCompiler.h"
 #include "TextureResource.h"
+#include "SPrimaryButton.h"
 
 #define LOCTEXT_NAMESPACE "PaperEditor"
 
@@ -186,7 +187,7 @@ void SPaperExtractSpritesDialog::Construct(const FArguments& InArgs, UTexture2D*
 	ChildSlot
 	[
 		SNew(SBorder)
-		.BorderImage(FAppStyle::GetBrush("DetailsView.CategoryTop"))
+		.BorderImage(FAppStyle::GetBrush("ToolPanel.GroupBorder"))
 		.Padding(FMargin(0.0f, 3.0f, 1.0f, 0.0f))
 		[
 			SNew(SHorizontalBox)
@@ -212,29 +213,25 @@ void SPaperExtractSpritesDialog::Construct(const FArguments& InArgs, UTexture2D*
 					DetailsPropertyView.ToSharedRef()
 				]
 				+ SVerticalBox::Slot()
-				.Padding(2.0f)
+				.Padding(FMargin(16))
 				.HAlign(HAlign_Right)
 				.AutoHeight()
 				[
 					SNew(SUniformGridPanel)
-					.SlotPadding(2)
+					.SlotPadding(FAppStyle::GetMargin("StandardDialog.SlotPadding"))
 					+ SUniformGridPanel::Slot(0, 0)
 					[
-						SNew(SButton)
-						.ButtonStyle(FAppStyle::Get(), "FlatButton.Success")
-						.ForegroundColor(FLinearColor::White)
-						.ContentPadding(FMargin(6, 2))
+						SNew(SPrimaryButton)
 						.Text(LOCTEXT("PaperExtractSpritesExtractButton", "Extract..."))
 						.OnClicked(this, &SPaperExtractSpritesDialog::ExtractClicked)
 					]
 					+ SUniformGridPanel::Slot(1, 0)
 					[
 						SNew(SButton)
-						.ButtonStyle(FAppStyle::Get(), "FlatButton")
-						.ForegroundColor(FLinearColor::White)
-						.ContentPadding(FMargin(6, 2))
 						.Text(LOCTEXT("PaperExtractSpritesCancelButton", "Cancel"))
-						.OnClicked(this, &SPaperExtractSpritesDialog::CancelClicked)
+                        .TextStyle(FAppStyle::Get(), "DialogButtonText")
+						.HAlign(HAlign_Center)
+                        .OnClicked(this, &SPaperExtractSpritesDialog::CancelClicked)
 					]
 				]
 			]
diff --git a/Engine/Plugins/Animation/DeformerGraph/Source/OptimusEditor/OptimusEditor.Build.cs b/Engine/Plugins/Animation/DeformerGraph/Source/OptimusEditor/OptimusEditor.Build.cs
--- a/Engine/Plugins/Animation/DeformerGraph/Source/OptimusEditor/OptimusEditor.Build.cs
+++ b/Engine/Plugins/Animation/DeformerGraph/Source/OptimusEditor/OptimusEditor.Build.cs
@@ -23,6 +23,7 @@ namespace UnrealBuildTool.Rules
 					"InputCore",
 					"RHI", 
 					"ToolMenus",
+					"ToolWidgets",
 					"ComputeFramework",
 					"OptimusCore",
 					"OptimusDeveloper",
diff --git a/Engine/Plugins/Animation/DeformerGraph/Source/OptimusEditor/Private/Widgets/SOptimusEditorGraphExplorer.cpp b/Engine/Plugins/Animation/DeformerGraph/Source/OptimusEditor/Private/Widgets/SOptimusEditorGraphExplorer.cpp
--- a/Engine/Plugins/Animation/DeformerGraph/Source/OptimusEditor/Private/Widgets/SOptimusEditorGraphExplorer.cpp
+++ b/Engine/Plugins/Animation/DeformerGraph/Source/OptimusEditor/Private/Widgets/SOptimusEditorGraphExplorer.cpp
@@ -24,6 +24,7 @@
 #include "Widgets/Input/SSearchBox.h"
 #include "Widgets/Input/SButton.h"
 #include "DetailLayoutBuilder.h"
+#include "SPositiveActionButton.h"
 
 // Ick.
 #include "GraphActionNode.h"
@@ -151,34 +152,12 @@ void SOptimusEditorGraphExplorer::RegisterCommands()
 
 void SOptimusEditorGraphExplorer::CreateWidgets()
 {
-	TSharedPtr<SWidget> AddNewMenu = SNew(SComboButton)
-		.ComboButtonStyle(FAppStyle::Get(), "ToolbarComboButton")
-		.ButtonStyle(FAppStyle::Get(), "FlatButton.Success")
-		.ForegroundColor(FLinearColor::White)
+	TSharedPtr<SWidget> AddNewMenu = SNew(SPositiveActionButton)
 		.ToolTipText(LOCTEXT("AddNewToolTip", "Add a new item."))
 		.OnGetMenuContent(this, &SOptimusEditorGraphExplorer::CreateAddNewMenuWidget)
-		.HasDownArrow(true)
-		.ContentPadding(FMargin(1, 0, 2, 0))
 		.IsEnabled(this, &SOptimusEditorGraphExplorer::IsEditingMode)
-		.ButtonContent()
-		[
-			SNew(SHorizontalBox)
-			+ SHorizontalBox::Slot()
-			.AutoWidth()
-			.Padding(FMargin(0, 1))
-			[
-				SNew(SImage)
-				.Image(FAppStyle::GetBrush("Plus"))
-			]
-			+ SHorizontalBox::Slot()
-			.VAlign(VAlign_Center)
-			.AutoWidth()
-			.Padding(FMargin(2, 0, 2, 0))
-			[
-				SNew(STextBlock)
-				.Text(LOCTEXT("AddNew", "Add New"))
-			]
-		];
+		.Icon(FAppStyle::GetBrush("Plus"))
+		.Text(LOCTEXT("AddNew", "Add"));
 
 	FMenuBuilder ViewOptions(true, nullptr);
 
@@ -249,10 +228,11 @@ void SOptimusEditorGraphExplorer::CreateWidgets()
 	                .Padding(2, 0, 0, 0)
 	                [
 						SNew(SComboButton)
+		                .ButtonStyle(FAppStyle::Get(), "SimpleButton")
 	                    .ComboButtonStyle(FAppStyle::Get(), "ToolbarComboButton")
 	                    .ForegroundColor(FSlateColor::UseForeground())
-	                    .HasDownArrow(true)
-	                    .ContentPadding(FMargin(1, 0))
+	                    .HasDownArrow(false)
+	                    .ContentPadding(0)
 	                    .AddMetaData<FTagMetaData>(FTagMetaData(TEXT("ViewOptions")))
 	                    .MenuContent()
 	                    [
@@ -260,8 +240,9 @@ void SOptimusEditorGraphExplorer::CreateWidgets()
 						]
 	                    .ButtonContent()
 	                    [
-							SNew(SImage)
-	                        .Image(FAppStyle::GetBrush("GenericViewButton"))
+	                    	SNew(SImage)
+							.Image(FAppStyle::GetBrush("Icons.Settings"))
+							.ColorAndOpacity(FSlateColor::UseForeground())
 						]
 					]
 				]
@@ -712,9 +693,9 @@ TSharedRef<SWidget> SOptimusEditorGraphExplorer::OnGetSectionWidget(TSharedRef<S
 	{
 		return SNew(SComboButton)
 			.ComboButtonStyle(FAppStyle::Get(), "ToolbarComboButton")
-			.ButtonStyle(FAppStyle::Get(), "RoundButton")
+			.ButtonStyle(FAppStyle::Get(), "SimpleButton")
+			.ContentPadding(FMargin(1, 0))
 		    .ForegroundColor(FAppStyle::GetSlateColor("DefaultForeground"))
-		    .ContentPadding(FMargin(2, 0))
 			.OnGetMenuContent_Lambda([AddMenuWidget]() { return AddMenuWidget.ToSharedRef(); })
 			.IsEnabled(this, &SOptimusEditorGraphExplorer::CanAddNewElementToSection, InSectionID)
 			.HasDownArrow(false)
@@ -723,23 +704,25 @@ TSharedRef<SWidget> SOptimusEditorGraphExplorer::OnGetSectionWidget(TSharedRef<S
 			.ButtonContent()
 			[
 				SNew(SImage)
-				.Image(FAppStyle::GetBrush("Plus"))
+				.Image(FAppStyle::GetBrush(TEXT("Icons.PlusCircle")))
+				.ColorAndOpacity(FSlateColor::UseForeground())
 				.ToolTipText(AddNewTooltipText)
 			];
 	}
 	else if (AddCommand.IsValid())
 	{
 		return SNew(SButton)
-		    .ButtonStyle(FAppStyle::Get(), "RoundButton")
+		    .ButtonStyle(FAppStyle::Get(), "SimpleButton")
+			.ContentPadding(FMargin(1, 0))
 		    .ForegroundColor(FAppStyle::GetSlateColor("DefaultForeground"))
-		    .ContentPadding(FMargin(2, 0))
 		    .OnClicked(this, &SOptimusEditorGraphExplorer::OnAddButtonClickedOnSection, InSectionID)
 		    .IsEnabled(this, &SOptimusEditorGraphExplorer::CanAddNewElementToSection, InSectionID)
 		    .HAlign(HAlign_Center)
 		    .VAlign(VAlign_Center)
 		    [
 				SNew(SImage)
-				.Image(FAppStyle::GetBrush("Plus"))
+		    	.Image(FAppStyle::GetBrush(TEXT("Icons.PlusCircle")))
+		    	.ColorAndOpacity(FSlateColor::UseForeground())
 		        .ToolTipText(AddNewTooltipText)
 			];
 	}
diff --git a/Engine/Plugins/Enterprise/DataprepEditor/Source/DataprepEditor/DataprepEditor.Build.cs b/Engine/Plugins/Enterprise/DataprepEditor/Source/DataprepEditor/DataprepEditor.Build.cs
--- a/Engine/Plugins/Enterprise/DataprepEditor/Source/DataprepEditor/DataprepEditor.Build.cs
+++ b/Engine/Plugins/Enterprise/DataprepEditor/Source/DataprepEditor/DataprepEditor.Build.cs
@@ -45,6 +45,7 @@ namespace UnrealBuildTool.Rules
 					"SlateCore",
 					"StatsViewer",
 					"ToolMenus",
+					"ToolWidgets",
 					"UnrealEd",
 					"SubobjectEditor",
 					"SubobjectDataInterface",
diff --git a/Engine/Plugins/Enterprise/DataprepEditor/Source/DataprepEditor/Private/Widgets/SDataprepProducersWidget.cpp b/Engine/Plugins/Enterprise/DataprepEditor/Source/DataprepEditor/Private/Widgets/SDataprepProducersWidget.cpp
--- a/Engine/Plugins/Enterprise/DataprepEditor/Source/DataprepEditor/Private/Widgets/SDataprepProducersWidget.cpp
+++ b/Engine/Plugins/Enterprise/DataprepEditor/Source/DataprepEditor/Private/Widgets/SDataprepProducersWidget.cpp
@@ -15,10 +15,12 @@
 #include "DetailLayoutBuilder.h"
 #include "DetailWidgetRow.h"
 #include "EditorFontGlyphs.h"
+#include "SPrimaryButton.h"
 #include "UObject/UObjectIterator.h"
 #include "Widgets/Input/SButton.h"
 #include "Widgets/Input/SComboButton.h"
 #include "Widgets/Layout/SScrollBox.h"
+
 #define LOCTEXT_NAMESPACE "DataprepProducersWidget"
 
 FContentProducerEntry::FContentProducerEntry(int32 InProducerIndex, UDataprepAssetProducers* InAssetProducersPtr)
@@ -254,10 +256,8 @@ void SDataprepProducersWidget::Construct( const FArguments & InArgs, UDataprepAs
 			.HAlign( HAlign_Center )
 			.Padding(20)
 			[
-				SNew(SButton)
-				.ButtonStyle(FAppStyle::Get(), "FlatButton.Success")
-				.ForegroundColor(FLinearColor::White)
-				.ContentPadding(FMargin(30, 2))
+				SNew(SPrimaryButton)
+				.Text( LOCTEXT( "ImportButton", "Import" ) )
 				.ToolTipText( LOCTEXT( "ImportButtonTooltip", "Load inputs' data into the Dataprep Editor"  ) )
 				.OnClicked(FOnClicked::CreateLambda([this]()
 				{
@@ -272,12 +272,6 @@ void SDataprepProducersWidget::Construct( const FArguments & InArgs, UDataprepAs
 					}
 					return false;
 				})
-				.Content()
-				[
-					SNew( STextBlock )
-					.TextStyle( FAppStyle::Get(), "ContentBrowser.TopBar.Font" )
-					.Text( LOCTEXT( "ImportButton", "Import" ) )
-				]
 			]
 		]
 		+SHorizontalBox::Slot()
diff --git a/Engine/Plugins/Enterprise/VariantManager/Source/VariantManager/Private/Dialogs/SCaptureDialog.cpp b/Engine/Plugins/Enterprise/VariantManager/Source/VariantManager/Private/Dialogs/SCaptureDialog.cpp
--- a/Engine/Plugins/Enterprise/VariantManager/Source/VariantManager/Private/Dialogs/SCaptureDialog.cpp
+++ b/Engine/Plugins/Enterprise/VariantManager/Source/VariantManager/Private/Dialogs/SCaptureDialog.cpp
@@ -13,6 +13,7 @@
 #include "VariantManagerPropertyCapturer.h"
 #include "GameFramework/GameUserSettings.h"
 #include "Application/SlateApplicationBase.h"
+#include "SPrimaryButton.h"
 
 #define LOCTEXT_NAMESPACE "SCaptureDialog"
 #define ACTOR_PROPERTY_DEFAULT_WIDTH 500.0f
@@ -176,23 +177,16 @@ void SCaptureDialog::Construct(const FArguments& InArgs)
 		.SlotPadding(FAppStyle::GetMargin("StandardDialog.SlotPadding"))
 		+SUniformGridPanel::Slot(0,0)
 		[
-			SNew(SButton)
+			SNew(SPrimaryButton)
 			.Text(NSLOCTEXT("SCapturedPropertiesWidget", "ClassPickerSelectButton", "Select"))
-			.HAlign(HAlign_Right)
-			.ContentPadding(FAppStyle::GetMargin("StandardDialog.ContentPadding"))
 			.OnClicked(this, &SCaptureDialog::OnDialogConfirmed)
-			.ButtonStyle(FAppStyle::Get(), "FlatButton.Success")
-			.TextStyle(FAppStyle::Get(), "FlatButton.DefaultTextStyle")
 		]
 		+SUniformGridPanel::Slot(1,0)
 		[
 			SNew(SButton)
 			.Text(NSLOCTEXT("SCapturedPropertiesWidget", "ClassPickerCancelButton", "Cancel"))
-			.HAlign(HAlign_Right)
-			.ContentPadding(FAppStyle::GetMargin("StandardDialog.ContentPadding"))
+			.TextStyle(FAppStyle::Get(), "DialogButtonText")
 			.OnClicked(this, &SCaptureDialog::OnDialogCanceled)
-			.ButtonStyle(FAppStyle::Get(), "FlatButton.Default")
-			.TextStyle(FAppStyle::Get(), "FlatButton.DefaultTextStyle")
 		]
 	];
 
@@ -225,7 +219,7 @@ void SCaptureDialog::Construct(const FArguments& InArgs)
 	[
 		SNew(SBorder)
 		.Visibility(EVisibility::Visible)
-		.BorderImage(FAppStyle::GetBrush("Menu.Background"))
+		.BorderImage(FAppStyle::GetBrush("Brushes.Panel"))
 		.Padding(FMargin(10.0f, 10.0f, 10.0f, 10.0f))
 		[
 			MainVerticalBox
diff --git a/Engine/Plugins/Enterprise/VariantManager/Source/VariantManager/Private/Dialogs/SCapturedPropertiesWidget.cpp b/Engine/Plugins/Enterprise/VariantManager/Source/VariantManager/Private/Dialogs/SCapturedPropertiesWidget.cpp
--- a/Engine/Plugins/Enterprise/VariantManager/Source/VariantManager/Private/Dialogs/SCapturedPropertiesWidget.cpp
+++ b/Engine/Plugins/Enterprise/VariantManager/Source/VariantManager/Private/Dialogs/SCapturedPropertiesWidget.cpp
@@ -21,6 +21,7 @@ void SCapturedPropertiesWidget::Construct(const FArguments& InArgs)
 		SNew(SExpandableArea)
 		.InitiallyCollapsed(false)
 		.Padding(0.0f)
+		.BorderImage(FAppStyle::Get().GetBrush("DetailsView.CategoryTop"))
 		.AreaTitle(NSLOCTEXT("SCapturedPropertiesWidget", "CapturedPropertiesText", "Captured Properties"))
 		.BodyContent()
 		[
diff --git a/Engine/Plugins/Experimental/GPULightmass/Source/GPULightmassEditor/GPULightmassEditor.Build.cs b/Engine/Plugins/Experimental/GPULightmass/Source/GPULightmassEditor/GPULightmassEditor.Build.cs
--- a/Engine/Plugins/Experimental/GPULightmass/Source/GPULightmassEditor/GPULightmassEditor.Build.cs
+++ b/Engine/Plugins/Experimental/GPULightmass/Source/GPULightmassEditor/GPULightmassEditor.Build.cs
@@ -7,16 +7,12 @@ namespace UnrealBuildTool.Rules
 		public GPULightmassEditor(ReadOnlyTargetRules Target) : base(Target)
 		{
 			PublicDependencyModuleNames.AddRange(
-				new string[]
-				{
+				new string[] {
 					"Core",
-					// ... add other public dependencies that you statically link with here ...
-				}
-				);
+				});
 
 			PrivateDependencyModuleNames.AddRange(
-				new string[]
-				{
+				new string[] {
 					"CoreUObject",
 					"Engine",
 					"Landscape",
@@ -31,8 +27,8 @@ namespace UnrealBuildTool.Rules
 					"WorkspaceMenuStructure",
 					"EditorStyle",
 					"ToolMenus",
-				}
-				);
+					"ToolWidgets",
+				});
 		}
 	}
 }
diff --git a/Engine/Plugins/Experimental/GPULightmass/Source/GPULightmassEditor/Private/GPULightmassEditorModule.cpp b/Engine/Plugins/Experimental/GPULightmass/Source/GPULightmassEditor/Private/GPULightmassEditorModule.cpp
--- a/Engine/Plugins/Experimental/GPULightmass/Source/GPULightmassEditor/Private/GPULightmassEditorModule.cpp
+++ b/Engine/Plugins/Experimental/GPULightmass/Source/GPULightmassEditor/Private/GPULightmassEditorModule.cpp
@@ -25,6 +25,9 @@
 #include "Widgets/Input/SButton.h"
 #include "Widgets/Input/SCheckBox.h"
 #include "ToolMenus.h"
+#include "SNegativeActionButton.h"
+#include "SPrimaryButton.h"
+#include "Widgets/Images/SImage.h"
 
 #define LOCTEXT_NAMESPACE "StaticLightingSystem"
 
@@ -239,50 +242,27 @@ TSharedRef<SDockTab> FGPULightmassEditorModule::SpawnSettingsTab(const FSpawnTab
 			SNew(SVerticalBox)
 			+ SVerticalBox::Slot()
 			.AutoHeight()
-			.Padding(2)
+			.Padding(12.0f, 8.0f)
 			[
-
 				// Start Build
 				SNew(SHorizontalBox)
 				+SHorizontalBox::Slot()
 				.AutoWidth()
 				.Padding(0.0f, 0.0f, 8.f, 0.0f)
 				[
-					SNew(SButton)
-					.HAlign(HAlign_Center)
-					.ButtonStyle(FAppStyle::Get(), "FlatButton.Success")
+					SNew(SPrimaryButton)
+					.Icon(FAppStyle::Get().GetBrush("EditorViewport.LightingOnlyMode"))
+					.Text_Lambda( []()
+					{ 
+						return FGPULightmassEditorModule::IsBakeWhatYouSeeMode() ? 
+								LOCTEXT("GPULightmassSettingsStartInteractive", "Start Building Lighting") : 
+								LOCTEXT("GPULightmassSettingsStartFull", "Build Lighting");
+					})
 					.IsEnabled_Lambda( []() {
 						return (GetRayTracingDisabledReason() == ERayTracingDisabledReason::OK) && IsStaticLightingAllowed() && IsPathTracingEnabled();
 					}) 
 					.Visibility_Lambda([](){ return IsRunning() ? EVisibility::Collapsed : EVisibility::Visible; })
 					.OnClicked_Raw(this, &FGPULightmassEditorModule::OnStartClicked)
-					[
-						SNew(SHorizontalBox)
-						+ SHorizontalBox::Slot()
-						.VAlign(VAlign_Center)
-						.AutoWidth()
-						[
-							SNew(STextBlock)
-							.TextStyle(FAppStyle::Get(), "ContentBrowser.TopBar.Font")
-							.Font(FAppStyle::Get().GetFontStyle("FontAwesome.11"))
-							.Text(FEditorFontGlyphs::Lightbulb_O)
-						]
-						+ SHorizontalBox::Slot()
-						.AutoWidth()
-						.VAlign(VAlign_Center)
-						.Padding(4, 0, 0, 0)
-						[
-							SNew(STextBlock)
-							.TextStyle(FAppStyle::Get(), "ContentBrowser.TopBar.Font")
-							.Text_Lambda( []()
-							{ 
-								return FGPULightmassEditorModule::IsBakeWhatYouSeeMode() ? 
-										LOCTEXT("GPULightmassSettingsStartInteractive", "Start Building Lighting") : 
-										LOCTEXT("GPULightmassSettingsStartFull", "Build Lighting");
-							})
-
-						]
-					]
 				]
 
 				// Save and Stop Building
@@ -290,32 +270,11 @@ TSharedRef<SDockTab> FGPULightmassEditorModule::SpawnSettingsTab(const FSpawnTab
 				.Padding(0.0f, 0.0f, 8.f, 0.0f)
 				.AutoWidth()
 				[
-					SNew(SButton)
-					.HAlign(HAlign_Center)
-					.ButtonStyle(FAppStyle::Get(), "FlatButton.Success")
+					SNew(SPrimaryButton)
+					.Icon(FAppStyle::Get().GetBrush("EditorViewport.LightingOnlyMode"))
+					.Text(LOCTEXT("GPULightmassSettingsSaveAndStop", "Save And Stop Building"))
 					.Visibility_Lambda([](){ return IsRunning() && IsBakeWhatYouSeeMode() ? EVisibility::Visible : EVisibility::Collapsed; })
 					.OnClicked_Raw(this, &FGPULightmassEditorModule::OnSaveAndStopClicked)
-					[
-						SNew(SHorizontalBox)
-						+ SHorizontalBox::Slot()
-						.VAlign(VAlign_Center)
-						.AutoWidth()
-						[
-							SNew(STextBlock)
-							.TextStyle(FAppStyle::Get(), "ContentBrowser.TopBar.Font")
-							.Font(FAppStyle::Get().GetFontStyle("FontAwesome.11"))
-							.Text(FEditorFontGlyphs::Lightbulb_O)
-						]
-						+ SHorizontalBox::Slot()
-						.AutoWidth()
-						.VAlign(VAlign_Center)
-						.Padding(4, 0, 0, 0)
-						[
-							SNew(STextBlock)
-							.TextStyle(FAppStyle::Get(), "ContentBrowser.TopBar.Font")
-							.Text(LOCTEXT("GPULightmassSettingsSaveAndStop", "Save And Stop Building"))
-						]
-					]
 				]
 
 				// Cancel Build
@@ -323,108 +282,171 @@ TSharedRef<SDockTab> FGPULightmassEditorModule::SpawnSettingsTab(const FSpawnTab
 				.Padding(0.0f, 0.0f, 8.f, 0.0f)
 				.AutoWidth()
 				[
-					SNew(SButton)
-					.HAlign(HAlign_Center)
-					.ButtonStyle(FAppStyle::Get(), "FlatButton.Danger")
+					SNew(SNegativeActionButton)
 					.Visibility_Lambda([](){ return IsRunning() ? EVisibility::Visible: EVisibility::Collapsed; })
 					.OnClicked_Raw(this, &FGPULightmassEditorModule::OnCancelClicked)
 					.Text(LOCTEXT("GPULightmassSettingsCancel", "Cancel Build"))
-					.TextStyle(FAppStyle::Get(), "ContentBrowser.TopBar.Font")
 				]
-				
 
 				+SHorizontalBox::Slot()
-				.FillWidth(1.0)
-				.HAlign(HAlign_Right)
-				.VAlign(VAlign_Center)
-				[
-					SNew(SCheckBox)
-					.IsChecked_Lambda( [] () { return FGPULightmassEditorModule::IsRealtimeOn() ? ECheckBoxState::Checked : ECheckBoxState::Unchecked; })
-					.OnCheckStateChanged_Lambda( [] (ECheckBoxState NewState) 
-					{
-						if (GCurrentLevelEditingViewportClient)
-						{
-							GCurrentLevelEditingViewportClient->SetRealtime( NewState == ECheckBoxState::Checked );
-						}	
-					})
-				]
-         
-				+SHorizontalBox::Slot()
-				.AutoWidth()
-				.HAlign(HAlign_Left)
 				.VAlign(VAlign_Center)
 				[
 					SNew(SBox)
-					.WidthOverride(140)
-					[	
-						SNew(STextBlock)
-						.Text_Lambda( [](){ return FGPULightmassEditorModule::IsRealtimeOn() ? LOCTEXT("GPULightmassRealtimeEnabled", "Viewport Realtime is ON ") : LOCTEXT("GPULightmassRealtimeDisabled", "Viewport Realtime is OFF");})
+					.ToolTipText(LOCTEXT("GPULightmassSettingsRealtimeToggleTooltip", "GPU Lightmass runs in slow mode when the viewport is realtime to avoid freezing. Uncheck Viewport Realtime to get full speed."))
+					[
+						SNew(SHorizontalBox)
+						+ SHorizontalBox::Slot()
+						.AutoWidth()
+						.HAlign(HAlign_Left)
+						.Padding(0, 0, 2, 0)
+						[
+							SNew(SCheckBox)
+							.IsChecked_Lambda( [] () { return FGPULightmassEditorModule::IsRealtimeOn() ? ECheckBoxState::Checked : ECheckBoxState::Unchecked; })
+							.OnCheckStateChanged_Lambda( [] (ECheckBoxState NewState) 
+							{
+								if (GCurrentLevelEditingViewportClient)
+								{
+									GCurrentLevelEditingViewportClient->SetRealtime( NewState == ECheckBoxState::Checked );
+								}
+							})
+						]
+
+						+ SHorizontalBox::Slot()
+						.AutoWidth()
+						.HAlign(HAlign_Left)
+						.Padding(0, 0, 4, 0)
+						[
+							SNew(SImage)
+							.Image(FAppStyle::GetBrush(TEXT("EditorViewport.ToggleRealTime")))
+						]
+
+						+ SHorizontalBox::Slot()
+						.AutoWidth()
+						.HAlign(HAlign_Left)
+						.Padding(0, 0, 2, 0)
+						[
+							SNew(STextBlock)
+							.Text(LOCTEXT("GPULightmassRealtimeEnabled", "Realtime Viewport"))
+						]
 					]
+
 				]
 			]
+
 			+ SVerticalBox::Slot()
 			.AutoHeight()
-			.Padding(2.f, 4.f)
+			.Padding(12.0f, 4.0f)
 			[
-				SAssignNew(Messages, STextBlock)
-				.AutoWrapText(true)
-				.Text_Lambda( []() -> FText
-				{
-					ERayTracingDisabledReason RayTracingStatus = GetRayTracingDisabledReason();
-					if (RayTracingStatus != ERayTracingDisabledReason::OK)
-					{
-						return GenerateRayTracingDisabledReasonMessage(RayTracingStatus);
-					}
-					else if (!IsStaticLightingAllowed())
+				SNew(SHorizontalBox)
+				+ SHorizontalBox::Slot()
+				.AutoWidth()
+				.HAlign(HAlign_Left)
+				.VAlign(VAlign_Top)
+				.Padding(0, 0, 4, 0)
+				[
+					SNew(SImage)
+					.Image_Lambda([this]() -> const FSlateBrush*
 					{
-						return LOCTEXT("GPULightmassAllowStaticLightingDisabled", "GPU Lightmass requires Allow Static Lighting enabled in the project settings.");
-					}
-					else if (!IsPathTracingEnabled())
+						// Error
+						if (Message.Severity <= ELogVerbosity::Error)
+						{
+							return FAppStyle::GetBrush(TEXT("Icons.ErrorWithColor"));
+						}
+						
+						// Warning
+						if (Message.Severity == ELogVerbosity::Warning)
+						{
+							return FAppStyle::GetBrush(TEXT("Icons.WarningWithColor"));
+						}
+						
+						// Info
+						return FAppStyle::GetBrush(TEXT("Icons.InfoWithColor"));
+					})
+				]
+
+				+ SHorizontalBox::Slot()
+				.FillWidth(1.0f)
+				.HAlign(HAlign_Fill)
+				.Padding(0, 0, 2, 0)
+				[
+					SAssignNew(MessageWidget, STextBlock)
+					.AutoWrapText(true)
+					.Text_Lambda( [this]() -> FText
 					{
-						return LOCTEXT("GPULightmassPathTracingDisabled", "GPU Lightmass requires Path Tracing enabled in the project settings.");
-					}
+						Message.Severity = ELogVerbosity::Error; // Default to error
+						
+						ERayTracingDisabledReason RayTracingStatus = GetRayTracingDisabledReason();
+						if (RayTracingStatus != ERayTracingDisabledReason::OK)
+						{
+							Message.Text = GenerateRayTracingDisabledReasonMessage(RayTracingStatus);
+							return Message.Text;
+						}
+						else if (!IsStaticLightingAllowed())
+						{
+							Message.Text = LOCTEXT("GPULightmassAllowStaticLightingDisabled", "GPU Lightmass requires Allow Static Lighting enabled in the project settings.");
+							return Message.Text;
+						}
+						else if (!IsPathTracingEnabled())
+						{
+							Message.Text = LOCTEXT("GPULightmassPathTracingDisabled", "GPU Lightmass requires Path Tracing enabled in the project settings.");
+							return Message.Text;
+						}
 
-					// Ready
-					static FText ReadyMsg = FText(LOCTEXT("GPULightmassReady", "GPU Lightmass is ready."));
+						// Message is info unless issue found later
+						Message.Severity = ELogVerbosity::Display;
 
-					// Ready, BWYS
-					static FText BWYSReadyMsg = FText(LOCTEXT("GPULightmassReadyBWYS", "GPU Lightmass is ready. Lighting will rebuild continuously in Bake What You See mode until saved or canceled."));
+						// Ready
+						static FText ReadyMsg = FText(LOCTEXT("GPULightmassReady", "GPU Lightmass is ready."));
 
-					// Ready, BWYS+RT OFF Warning 
-					static FText RtOffBWYSWarningMsg = LOCTEXT("GPULightmassSpeedReadyRTWarning", "Building Lighting when using Bake What You See Mode will automatically enable Viewport Realtime to start building. Lighting will rebuild continuously in Bake What You See mode until saved or canceled.");
+						// Ready, BWYS
+						static FText BWYSReadyMsg = FText(LOCTEXT("GPULightmassReadyBWYS", "GPU Lightmass is ready. Lighting will rebuild continuously in Bake What You See mode until saved or canceled."));
 
-					// Building FULL + RT Off Warning
-					UWorld* World = GEditor->GetEditorWorldContext().World();
-					FText BuildingMsg = FText::Format(LOCTEXT("GPULightmassBuildingLighting", "GPU Lightmass is building lighting for {0}."), FText::FromString(World->GetActiveLightingScenario() ? World->GetActiveLightingScenario()->GetOuter()->GetName() : World->GetName()));
+						// Ready, BWYS+RT OFF Warning 
+						static FText RtOffBWYSWarningMsg = LOCTEXT("GPULightmassSpeedReadyRTWarning", "Building Lighting when using Bake What You See Mode will automatically enable Viewport Realtime to start building. Lighting will rebuild continuously in Bake What You See mode until saved or canceled.");
 
-					// Building FULL + RT ON Warning 
-					static FText BuildingFullRTOnMsg = LOCTEXT("GPULightmassBuildingFullRTOn", "GPU Lightmass runs in slow mode when the viewport is realtime to avoid freezing. Uncheck Viewport Realtime to get full speed.");
+						// Building FULL + RT Off Warning
+						UWorld* World = GEditor->GetEditorWorldContext().World();
+						FText BuildingMsg = FText::Format(LOCTEXT("GPULightmassBuildingLighting", "GPU Lightmass is building lighting for {0}."), FText::FromString(World->GetActiveLightingScenario() ? World->GetActiveLightingScenario()->GetOuter()->GetName() : World->GetName()));
 
-					// Building BWYS + RT ON Warning 
-					static FText BuildingRTOnMsg = LOCTEXT("GPULightmassBuildingInteractiveRTOn", "Disable Viewport Realtime to speed up building.");
+						// Building FULL + RT ON Warning 
+						static FText BuildingFullRTOnMsg = LOCTEXT("GPULightmassBuildingFullRTOn", "GPU Lightmass runs in slow mode when the viewport is realtime to avoid freezing. Uncheck Viewport Realtime to get full speed.");
 
-					// Building BWYS + RT OFF Warning 
-					static FText BuidlingRTOffMsg = LOCTEXT("GPULightmassBuildingInteractiveRTOff", "Re-enable Viewport Realtime to preview lighting.  Enabling Viewport Realtime will slow down building, to avoid freezing.");
+						// Building BWYS + RT ON Warning 
+						static FText BuildingRTOnMsg = LOCTEXT("GPULightmassBuildingInteractiveRTOn", "Disable Viewport Realtime to speed up building.");
 
-					bool bIsRunning = IsRunning();
-					bool bIsInteractive = IsBakeWhatYouSeeMode();
-					bool bIsRealtime = IsRealtimeOn();
-					if (bIsRunning)
-					{
-						if (bIsInteractive)
+						// Building BWYS + RT OFF Warning 
+						static FText BuidlingRTOffMsg = LOCTEXT("GPULightmassBuildingInteractiveRTOff", "Re-enable Viewport Realtime to preview lighting.  Enabling Viewport Realtime will slow down building, to avoid freezing.");
+
+						bool bIsRunning = IsRunning();
+						bool bIsInteractive = IsBakeWhatYouSeeMode();
+						bool bIsRealtime = IsRealtimeOn();
+						if (bIsRunning)
 						{
-							return bIsRealtime ? BuildingRTOnMsg : BuidlingRTOffMsg;
+							if (bIsInteractive)
+							{
+								// Both messages below are considered warnings
+								Message.Severity = ELogVerbosity::Warning;
+								Message.Text = bIsRealtime ? BuildingRTOnMsg : BuidlingRTOffMsg;
+								return Message.Text;
+							}
+
+							Message.Severity = bIsRealtime ? ELogVerbosity::Warning : ELogVerbosity::Display;
+							Message.Text = bIsRealtime ? BuildingFullRTOnMsg : BuildingMsg;
+							return Message.Text;
+						}
+						else if (bIsInteractive)
+						{
+							// Both messages below are considered warnings
+							Message.Severity = ELogVerbosity::Warning;
+							Message.Text = bIsRealtime ?  BWYSReadyMsg : RtOffBWYSWarningMsg;
+							return Message.Text;
 						}
 
-						return bIsRealtime ? BuildingFullRTOnMsg : BuildingMsg;
-					}
-					else if (bIsInteractive)
-					{
-						return bIsRealtime ?  BWYSReadyMsg : RtOffBWYSWarningMsg;
-					}
-
-					return bIsRealtime ? BuildingFullRTOnMsg : ReadyMsg;
-				})	
+						Message.Severity = bIsRealtime ? ELogVerbosity::Warning : ELogVerbosity::Display;
+						Message.Text = bIsRealtime ? BuildingFullRTOnMsg : ReadyMsg;
+						return Message.Text;
+					})
+				]
 			]
 
 			+ SVerticalBox::Slot()
diff --git a/Engine/Plugins/Experimental/GPULightmass/Source/GPULightmassEditor/Private/GPULightmassEditorModule.h b/Engine/Plugins/Experimental/GPULightmass/Source/GPULightmassEditor/Private/GPULightmassEditorModule.h
--- a/Engine/Plugins/Experimental/GPULightmass/Source/GPULightmassEditor/Private/GPULightmassEditorModule.h
+++ b/Engine/Plugins/Experimental/GPULightmass/Source/GPULightmassEditor/Private/GPULightmassEditorModule.h
@@ -41,5 +41,13 @@ public:
 	FReply OnSaveAndStopClicked();
 	FReply OnCancelClicked();
 
-	TSharedPtr<STextBlock> Messages;
+private:
+	struct FMessage
+	{
+		ELogVerbosity::Type Severity;
+		FText Text;
+	};
+
+	FMessage Message = {};
+	TSharedPtr<STextBlock> MessageWidget;
 };
diff --git a/Engine/Source/Editor/Blutility/Blutility.Build.cs b/Engine/Source/Editor/Blutility/Blutility.Build.cs
--- a/Engine/Source/Editor/Blutility/Blutility.Build.cs
+++ b/Engine/Source/Editor/Blutility/Blutility.Build.cs
@@ -41,6 +41,7 @@ public class Blutility : ModuleRules
                 "UMGEditor",
                 "KismetCompiler",
 				"ToolMenus",
+				"ToolWidgets",
 				"RHI",
 				"RenderCore",
 				"ImageWrapper",
diff --git a/Engine/Source/Editor/Blutility/Private/BlutilityMenuExtensions.cpp b/Engine/Source/Editor/Blutility/Private/BlutilityMenuExtensions.cpp
--- a/Engine/Source/Editor/Blutility/Private/BlutilityMenuExtensions.cpp
+++ b/Engine/Source/Editor/Blutility/Private/BlutilityMenuExtensions.cpp
@@ -47,6 +47,7 @@
 #include "PropertyEditorModule.h"
 #include "ScopedTransaction.h"
 #include "SlotBase.h"
+#include "SPrimaryButton.h"
 #include "Styling/AppStyle.h"
 #include "Styling/SlateColor.h"
 #include "Subsystems/AssetEditorSubsystem.h"
@@ -497,18 +498,17 @@ class SFunctionParamDialog : public SCompoundWidget
 			[
 				SNew(SBorder)
 				.BorderImage(FAppStyle::GetBrush("ToolPanel.GroupBorder"))
+				.Padding(16.0f)
 				.VAlign(VAlign_Center)
 				.HAlign(HAlign_Right)
 				[
 					SNew(SHorizontalBox)
 					+SHorizontalBox::Slot()
-					.Padding(2.0f)
 					.AutoWidth()
 					[
-						SNew(SButton)
-						.ButtonStyle(FAppStyle::Get(), "FlatButton.Success")
-						.ForegroundColor(FLinearColor::White)
-						.ContentPadding(FMargin(6, 2))
+						SNew(SPrimaryButton)
+						.Text(InArgs._OkButtonText)
+						.ToolTipText(InArgs._OkButtonTooltipText)
 						.OnClicked_Lambda([this, InParentWindow, InArgs]()
 						{
 							if(InParentWindow.IsValid())
@@ -518,34 +518,22 @@ class SFunctionParamDialog : public SCompoundWidget
 							bOKPressed = true;
 							return FReply::Handled(); 
 						})
-						.ToolTipText(InArgs._OkButtonTooltipText)
-						[
-							SNew(STextBlock)
-							.TextStyle(FAppStyle::Get(), "ContentBrowser.TopBar.Font")
-							.Text(InArgs._OkButtonText)
-						]
 					]
 					+SHorizontalBox::Slot()
-					.Padding(2.0f)
+					.Padding(FMargin(8.0f, 0.0f, 0.0f, 0.0f))
 					.AutoWidth()
 					[
 						SNew(SButton)
-						.ButtonStyle(FAppStyle::Get(), "FlatButton")
-						.ForegroundColor(FLinearColor::White)
-						.ContentPadding(FMargin(6, 2))
+						.TextStyle(FAppStyle::Get(), "DialogButtonText")
+						.Text(LOCTEXT("Cancel", "Cancel"))
 						.OnClicked_Lambda([InParentWindow]()
-						{ 
+						{
 							if(InParentWindow.IsValid())
 							{
 								InParentWindow.Pin()->RequestDestroyWindow();
 							}
 							return FReply::Handled(); 
 						})
-						[
-							SNew(STextBlock)
-							.TextStyle(FAppStyle::Get(), "ContentBrowser.TopBar.Font")
-							.Text(LOCTEXT("Cancel", "Cancel"))
-						]
 					]
 				]
 			]
diff --git a/Engine/Source/Runtime/SlateCore/Private/Styling/StarshipCoreStyle.cpp b/Engine/Source/Runtime/SlateCore/Private/Styling/StarshipCoreStyle.cpp
--- a/Engine/Source/Runtime/SlateCore/Private/Styling/StarshipCoreStyle.cpp
+++ b/Engine/Source/Runtime/SlateCore/Private/Styling/StarshipCoreStyle.cpp
@@ -802,8 +802,8 @@ TSharedRef<ISlateStyle> FStarshipCoreStyle::Create()
 
 	// Standard Dialog Settings
 	{
-		Style->Set("StandardDialog.ContentPadding", FMargin(12.0f, 2.0f));
-		Style->Set("StandardDialog.SlotPadding", FMargin(6.0f, 0.0f, 0.0f, 0.0f));
+		Style->Set("StandardDialog.ContentPadding", FMargin(16.0f, 3.0f));
+		Style->Set("StandardDialog.SlotPadding", FMargin(8.0f, 0.0f, 0.0f, 0.0f));
 		Style->Set("StandardDialog.MinDesiredSlotWidth", 80.0f);
 		Style->Set("StandardDialog.MinDesiredSlotHeight", 0.0f);
 		Style->Set("StandardDialog.SmallFont", StyleFonts.Small);
