[GPULM] Fix crashes after recent changes:
* Specify RayTracingScene when creating an SRV
* Create a buffer for SceneUniforms
* Get geometry from RayTracingProxy

#rnx

diff --git a/Engine/Plugins/Experimental/GPULightmass/Source/GPULightmass/Private/LightmapRenderer.cpp b/Engine/Plugins/Experimental/GPULightmass/Source/GPULightmass/Private/LightmapRenderer.cpp
--- a/Engine/Plugins/Experimental/GPULightmass/Source/GPULightmass/Private/LightmapRenderer.cpp
+++ b/Engine/Plugins/Experimental/GPULightmass/Source/GPULightmass/Private/LightmapRenderer.cpp
@@ -705,7 +705,7 @@ void FCachedRayTracingSceneData::SetupFromSceneRenderState(FSceneRenderState& Sc
 
 				int32 InstanceIndex = RayTracingGeometryInstancesPerLOD[LODIndex].AddDefaulted(1);
 				FRayTracingGeometryInstance& RayTracingInstance = RayTracingGeometryInstancesPerLOD[LODIndex][InstanceIndex];
-				RayTracingInstance.GeometryRHI = Instance.RenderData->LODResources[LODIndexToUse].RayTracingGeometry->GetRHI();
+				RayTracingInstance.GeometryRHI = Instance.RenderData->RayTracingProxy->LODs[LODIndexToUse].RayTracingGeometry->GetRHI();
 				RayTracingInstance.Transforms = MakeArrayView(&Instance.LocalToWorld, 1);
 				RayTracingInstance.NumTransforms = 1;
 				RayTracingInstance.InstanceContributionToHitGroupIndex = RayTracing::CalculateInstanceContributionToHitGroupIndex(GlobalSegmentIndex);
@@ -829,6 +829,9 @@ bool FSceneRenderState::SetupRayTracingScene(FRDGBuilder& GraphBuilder, FSceneUn
 
 	FRHICommandListImmediate& RHICmdList = FRHICommandListExecutor::GetImmediateCommandList();
 
+	// Make sure the buffer is prepped
+	SceneUniforms.GetBuffer(GraphBuilder);
+
 #if RHI_RAYTRACING
 	// Force build all the open build requests
 	bool bBuildAll = true;
@@ -1140,7 +1143,7 @@ bool FSceneRenderState::SetupRayTracingScene(FRDGBuilder& GraphBuilder, FSceneUn
 			FRayTracingAccelerationStructureSize SizeInfo = RHICalcRayTracingSceneSize(SceneInitializer.NumNativeInstances, ERayTracingAccelerationStructureFlags::FastTrace);
 			FRHIResourceCreateInfo BufferCreateInfo(TEXT("LightmassRayTracingSceneBuffer"));
 			RayTracingSceneBuffer = RHICmdList.CreateBuffer(uint32(SizeInfo.ResultSize), BUF_AccelerationStructure, 0, ERHIAccess::BVHWrite, BufferCreateInfo);
-			RayTracingSceneSRV = RHICmdList.CreateShaderResourceView(RayTracingSceneBuffer);
+			RayTracingSceneSRV = RHICmdList.CreateShaderResourceView(FShaderResourceViewInitializer(RayTracingSceneBuffer, RayTracingScene, 0));
 
 			FRHIResourceCreateInfo ScratchBufferCreateInfo(TEXT("LightmassRayTracingScratchBuffer"));
 			FBufferRHIRef ScratchBuffer = RHICmdList.CreateBuffer(
@@ -2077,8 +2080,7 @@ void FLightmapRenderer::Finalize(FRDGBuilder& GraphBuilder)
 	RectLightAtlas::UpdateAtlasTexture(GraphBuilder, Scene->FeatureLevel);
 	IESAtlas::UpdateAtlasTexture(GraphBuilder, GetFeatureLevelShaderPlatform(Scene->FeatureLevel));
 
-	FSceneUniformBuffer SceneUniforms;
-
+	FSceneUniformBuffer SceneUniforms{};
 	if (!Scene->SetupRayTracingScene(GraphBuilder, SceneUniforms, MostCommonLODIndex))
 	{
 		return;
