Use Proxy to retrieve RayTracingGeometry

diff --git a/Engine/Plugins/Experimental/GPULightmass/Source/GPULightmass/Private/LightmapRenderer.cpp b/Engine/Plugins/Experimental/GPULightmass/Source/GPULightmass/Private/LightmapRenderer.cpp
--- a/Engine/Plugins/Experimental/GPULightmass/Source/GPULightmass/Private/LightmapRenderer.cpp
+++ b/Engine/Plugins/Experimental/GPULightmass/Source/GPULightmass/Private/LightmapRenderer.cpp
@@ -748,7 +748,7 @@ void FCachedRayTracingSceneData::SetupFromSceneRenderState(FSceneRenderState& Sc
 
 				int32 InstanceIndex = RayTracingGeometryInstancesPerLOD[LODIndex].AddDefaulted(1);
 				FRayTracingGeometryInstance& RayTracingInstance = RayTracingGeometryInstancesPerLOD[LODIndex][InstanceIndex];
-				RayTracingInstance.GeometryRHI = InstanceGroup.ComponentUObject->GetStaticMesh()->GetRenderData()->LODResources[LODIndexToUse].RayTracingGeometry->GetRHI();
+				RayTracingInstance.GeometryRHI = InstanceGroup.ComponentUObject->GetStaticMesh()->GetRenderData()->RayTracingProxy->LODs[LODIndexToUse].RayTracingGeometry->GetRHI();
 
 				const int32 NumInstances = (int32)InstanceGroup.NumInstances;
 				TArrayView<FMatrix> NewTransforms = MakeArrayView(
