[Persistent FRHIRayTracingScene] Replace CreateRayTracingSceneWithGeometryInstances(...) with BuildRayTracingSceneInitializationData(...)
- this helper function is primarly intended to prepare the data required to build the ray tracing instance buffer (FillRayTracingInstanceUploadBuffer(...)/BuildRayTracingInstanceBuffer(...))
- instead of creating a new FRHIRayTracingScene inside this helper function, higher level code should instead create or reuse existing FRHIRayTracingScene objects.

diff --git a/Engine/Plugins/Experimental/GPULightmass/Source/GPULightmass/Private/LightmapRenderer.cpp b/Engine/Plugins/Experimental/GPULightmass/Source/GPULightmass/Private/LightmapRenderer.cpp
--- a/Engine/Plugins/Experimental/GPULightmass/Source/GPULightmass/Private/LightmapRenderer.cpp
+++ b/Engine/Plugins/Experimental/GPULightmass/Source/GPULightmass/Private/LightmapRenderer.cpp
@@ -1134,15 +1134,22 @@ bool FSceneRenderState::SetupRayTracingScene(FRDGBuilder& GraphBuilder, FSceneUn
 		{
 			SCOPED_GPU_MASK(RHICmdList, FRHIGPUMask::All());
 
-			FRayTracingSceneWithGeometryInstances SceneWithGeometryInstances = CreateRayTracingSceneWithGeometryInstances(
-				RayTracingGeometryInstances,
-				ERayTracingAccelerationStructureFlags::FastTrace);
+			const FRayTracingSceneInitializationData RayTracingSceneInitializationData = BuildRayTracingSceneInitializationData(RayTracingGeometryInstances);
+			const ERayTracingAccelerationStructureFlags SceneBuildFlags = ERayTracingAccelerationStructureFlags::FastTrace;
 
-			RayTracingScene = SceneWithGeometryInstances.Scene;
+			{
+				FRayTracingSceneInitializer2 Initializer;
+				Initializer.DebugName = FName(TEXT("LightmapRendererRayTracingScene"));
+				Initializer.NumNativeInstances = RayTracingSceneInitializationData.NumNativeGPUSceneInstances + RayTracingSceneInitializationData.NumNativeCPUInstances;
+				Initializer.NumTotalSegments = RayTracingSceneInitializationData.TotalNumSegments;
+				Initializer.BuildFlags = SceneBuildFlags;
+
+				RayTracingScene = RHICreateRayTracingScene(MoveTemp(Initializer));
+			}
 
 			const FRayTracingSceneInitializer2& SceneInitializer = RayTracingScene->GetInitializer();
 
-			FRayTracingAccelerationStructureSize SizeInfo = RHICalcRayTracingSceneSize(SceneInitializer.NumNativeInstances, ERayTracingAccelerationStructureFlags::FastTrace);
+			FRayTracingAccelerationStructureSize SizeInfo = RHICalcRayTracingSceneSize(SceneInitializer.NumNativeInstances, SceneBuildFlags);
 			FRHIResourceCreateInfo BufferCreateInfo(TEXT("LightmassRayTracingSceneBuffer"));
 			RayTracingSceneBuffer = RHICmdList.CreateBuffer(uint32(SizeInfo.ResultSize), BUF_AccelerationStructure, 0, ERHIAccess::BVHWrite, BufferCreateInfo);
 			RayTracingSceneSRV = RHICmdList.CreateShaderResourceView(FShaderResourceViewInitializer(RayTracingSceneBuffer, RayTracingScene, 0));
@@ -1161,7 +1168,7 @@ bool FSceneRenderState::SetupRayTracingScene(FRDGBuilder& GraphBuilder, FSceneUn
 			AccelerationStructureAddressesBuffer.Initialize(
 				RHICmdList,
 				TEXT("LightmassRayTracingAccelerationStructureAddressesBuffer"),
-				SceneWithGeometryInstances.ReferencedGeometries.Num() * sizeof(FRayTracingAccelerationStructureAddress),
+				RayTracingSceneInitializationData.ReferencedGeometries.Num() * sizeof(FRayTracingAccelerationStructureAddress),
 				BUF_Volatile | BUF_MultiGPUAllocate);
 
 			const uint32 InstanceUploadBufferSize = SceneInitializer.NumNativeInstances * sizeof(FRayTracingInstanceDescriptorInput);
@@ -1173,7 +1180,7 @@ bool FSceneRenderState::SetupRayTracingScene(FRDGBuilder& GraphBuilder, FSceneUn
 				InstanceUploadSRV = RHICmdList.CreateShaderResourceView(InstanceUploadBuffer);
 			}
 
-			const uint32 TransformUploadBufferSize = SceneWithGeometryInstances.NumNativeCPUInstances * 3 * sizeof(FVector4f);
+			const uint32 TransformUploadBufferSize = RayTracingSceneInitializationData.NumNativeCPUInstances * 3 * sizeof(FVector4f);
 			FBufferRHIRef TransformUploadBuffer;
 			FShaderResourceViewRHIRef TransformUploadSRV;
 			{
@@ -1189,13 +1196,13 @@ bool FSceneRenderState::SetupRayTracingScene(FRDGBuilder& GraphBuilder, FSceneUn
 					RayTracingScene,
 					View.ViewMatrices.GetPreViewTranslation(),
 					RayTracingGeometryInstances,
-					SceneWithGeometryInstances.InstanceGeometryIndices,
-					SceneWithGeometryInstances.BaseUploadBufferOffsets,
-					SceneWithGeometryInstances.BaseInstancePrefixSum,
-					SceneWithGeometryInstances.NumNativeGPUSceneInstances,
-					SceneWithGeometryInstances.NumNativeCPUInstances,
+					RayTracingSceneInitializationData.InstanceGeometryIndices,
+					RayTracingSceneInitializationData.BaseUploadBufferOffsets,
+					RayTracingSceneInitializationData.BaseInstancePrefixSum,
+					RayTracingSceneInitializationData.NumNativeGPUSceneInstances,
+					RayTracingSceneInitializationData.NumNativeCPUInstances,
 					MakeArrayView(InstanceUploadData, SceneInitializer.NumNativeInstances),
-					MakeArrayView(TransformUploadData, SceneWithGeometryInstances.NumNativeCPUInstances * 3));
+					MakeArrayView(TransformUploadData, RayTracingSceneInitializationData.NumNativeCPUInstances * 3));
 				RHICmdList.UnlockBuffer(TransformUploadBuffer);
 				RHICmdList.UnlockBuffer(InstanceUploadBuffer);
 			}
@@ -1206,9 +1213,9 @@ bool FSceneRenderState::SetupRayTracingScene(FRDGBuilder& GraphBuilder, FSceneUn
 					AccelerationStructureAddressesBuffer.Buffer,
 					GPUIndex,
 					0,
-					SceneWithGeometryInstances.ReferencedGeometries.Num() * sizeof(FRayTracingAccelerationStructureAddress), RLM_WriteOnly);
+					RayTracingSceneInitializationData.ReferencedGeometries.Num() * sizeof(FRayTracingAccelerationStructureAddress), RLM_WriteOnly);
 
-				const TArrayView<FRHIRayTracingGeometry*> ReferencedGeometries = RHICmdList.AllocArray(MakeConstArrayView(SceneWithGeometryInstances.ReferencedGeometries));
+				const TArrayView<FRHIRayTracingGeometry*> ReferencedGeometries = RHICmdList.AllocArray(MakeConstArrayView(RayTracingSceneInitializationData.ReferencedGeometries));
 
 				RHICmdList.EnqueueLambda([AddressesPtr, ReferencedGeometries, GPUIndex](FRHICommandListBase&)
 				{
@@ -1231,8 +1238,8 @@ bool FSceneRenderState::SetupRayTracingScene(FRDGBuilder& GraphBuilder, FSceneUn
 				InstanceUploadSRV,
 				AccelerationStructureAddressesBuffer.SRV,
 				TransformUploadSRV,
-				SceneWithGeometryInstances.NumNativeGPUSceneInstances,
-				SceneWithGeometryInstances.NumNativeCPUInstances,
+				RayTracingSceneInitializationData.NumNativeGPUSceneInstances,
+				RayTracingSceneInitializationData.NumNativeCPUInstances,
 				nullptr,
 				nullptr,
 				nullptr);
@@ -1246,8 +1253,8 @@ bool FSceneRenderState::SetupRayTracingScene(FRDGBuilder& GraphBuilder, FSceneUn
 				BuildParams.ScratchBufferOffset = 0;
 				BuildParams.InstanceBuffer = InstanceBuffer.Buffer;
 				BuildParams.InstanceBufferOffset = 0;
-				BuildParams.ReferencedGeometries = SceneWithGeometryInstances.ReferencedGeometries;
-				BuildParams.PerInstanceGeometries = SceneWithGeometryInstances.PerInstanceGeometries;
+				BuildParams.ReferencedGeometries = RayTracingSceneInitializationData.ReferencedGeometries;
+				BuildParams.PerInstanceGeometries = RayTracingSceneInitializationData.PerInstanceGeometries;
 
 				RHICmdList.BuildAccelerationStructure(BuildParams);
 			}
@@ -1303,7 +1310,7 @@ bool FSceneRenderState::SetupRayTracingScene(FRDGBuilder& GraphBuilder, FSceneUn
 
 			FRayTracingShaderBindingTableInitializer SBTInitializer;
 			SBTInitializer.bAllowHitGroupIndexing = true;
-			SBTInitializer.NumGeometrySegments = SceneWithGeometryInstances.TotalNumSegments;
+			SBTInitializer.NumGeometrySegments = RayTracingSceneInitializationData.TotalNumSegments;
 			SBTInitializer.NumShaderSlotsPerGeometrySegment = RAY_TRACING_NUM_SHADER_SLOTS;
 			SBTInitializer.NumMissShaderSlots = 1;
 			SBTInitializer.NumCallableShaderSlots = 0;
diff --git a/Engine/Source/Runtime/Renderer/Private/RayTracing/RayTracingInstanceBufferUtil.cpp b/Engine/Source/Runtime/Renderer/Private/RayTracing/RayTracingInstanceBufferUtil.cpp
--- a/Engine/Source/Runtime/Renderer/Private/RayTracing/RayTracingInstanceBufferUtil.cpp
+++ b/Engine/Source/Runtime/Renderer/Private/RayTracing/RayTracingInstanceBufferUtil.cpp
@@ -25,6 +25,76 @@
 
 #if RHI_RAYTRACING
 
+FRayTracingSceneInitializationData BuildRayTracingSceneInitializationData(TConstArrayView<FRayTracingGeometryInstance> Instances)
+{
+	const uint32 NumSceneInstances = Instances.Num();
+	
+	FRayTracingSceneInitializationData Output;
+	Output.NumNativeGPUSceneInstances = 0;
+	Output.NumNativeCPUInstances = 0;
+	Output.TotalNumSegments = 0;
+	Output.InstanceGeometryIndices.SetNumUninitialized(NumSceneInstances);
+	Output.BaseUploadBufferOffsets.SetNumUninitialized(NumSceneInstances);
+	Output.BaseInstancePrefixSum.SetNumUninitialized(NumSceneInstances);
+	Output.PerInstanceGeometries.SetNumUninitialized(NumSceneInstances);
+
+	Experimental::TSherwoodMap<FRHIRayTracingGeometry*, uint32> UniqueGeometries;
+
+	// Compute geometry segment prefix sum used by GetHitRecordBaseIndex() during resource binding
+
+	uint32 NumNativeInstances = 0;
+
+	for (uint32 InstanceIndex = 0; InstanceIndex < NumSceneInstances; ++InstanceIndex)
+	{
+		const FRayTracingGeometryInstance& InstanceDesc = Instances[InstanceIndex];
+
+		const bool bGpuSceneInstance = InstanceDesc.BaseInstanceSceneDataOffset != -1 || !InstanceDesc.InstanceSceneDataOffsets.IsEmpty();
+		const bool bCpuInstance = !bGpuSceneInstance;
+
+		checkf(!bGpuSceneInstance || InstanceDesc.BaseInstanceSceneDataOffset != -1 || InstanceDesc.NumTransforms <= uint32(InstanceDesc.InstanceSceneDataOffsets.Num()),
+			TEXT("Expected at least %d ray tracing geometry instance scene data offsets, but got %d."),
+			InstanceDesc.NumTransforms, InstanceDesc.InstanceSceneDataOffsets.Num());
+		checkf(!bCpuInstance || InstanceDesc.NumTransforms <= uint32(InstanceDesc.Transforms.Num()),
+			TEXT("Expected at least %d ray tracing geometry instance transforms, but got %d."),
+			InstanceDesc.NumTransforms, InstanceDesc.Transforms.Num());
+
+		checkf(InstanceDesc.GeometryRHI, TEXT("Ray tracing instance must have a valid geometry."));
+
+		Output.PerInstanceGeometries[InstanceIndex] = InstanceDesc.GeometryRHI;
+
+		Output.TotalNumSegments += InstanceDesc.GeometryRHI->GetNumSegments();
+
+		uint32 GeometryIndex = UniqueGeometries.FindOrAdd(InstanceDesc.GeometryRHI, Output.ReferencedGeometries.Num());
+		Output.InstanceGeometryIndices[InstanceIndex] = GeometryIndex;
+		if (GeometryIndex == Output.ReferencedGeometries.Num())
+		{
+			Output.ReferencedGeometries.Add(InstanceDesc.GeometryRHI);
+		}
+
+		if (bGpuSceneInstance)
+		{
+			check(InstanceDesc.Transforms.IsEmpty());
+			Output.BaseUploadBufferOffsets[InstanceIndex] = Output.NumNativeGPUSceneInstances;
+			Output.NumNativeGPUSceneInstances += InstanceDesc.NumTransforms;
+		}
+		else if (bCpuInstance)
+		{
+			Output.BaseUploadBufferOffsets[InstanceIndex] = Output.NumNativeCPUInstances;
+			Output.NumNativeCPUInstances += InstanceDesc.NumTransforms;
+		}
+		else
+		{
+			checkNoEntry();
+		}
+		
+		Output.BaseInstancePrefixSum[InstanceIndex] = NumNativeInstances;
+
+		NumNativeInstances += InstanceDesc.NumTransforms;
+	}
+
+	return MoveTemp(Output);
+}
+
 FRayTracingSceneWithGeometryInstances CreateRayTracingSceneWithGeometryInstances(
 	TConstArrayView<FRayTracingGeometryInstance> Instances,
 	uint8 NumLayers,
@@ -138,15 +208,6 @@ FRayTracingSceneWithGeometryInstances CreateRayTracingSceneWithGeometryInstances
 	return MoveTemp(Output);
 }
 
-FRayTracingSceneWithGeometryInstances CreateRayTracingSceneWithGeometryInstances(
-	TConstArrayView<FRayTracingGeometryInstance> Instances,
-	ERayTracingAccelerationStructureFlags BuildFlags)
-{
-	PRAGMA_DISABLE_DEPRECATION_WARNINGS
-	return CreateRayTracingSceneWithGeometryInstances(Instances, 1, 1, 1, 0, BuildFlags);
-	PRAGMA_ENABLE_DEPRECATION_WARNINGS
-}
-
 void FillRayTracingInstanceUploadBuffer(
 	FRayTracingSceneRHIRef RayTracingSceneRHI,
 	FVector PreViewTranslation,
@@ -161,8 +222,6 @@ void FillRayTracingInstanceUploadBuffer(
 {
 	TRACE_CPUPROFILER_EVENT_SCOPE(FillRayTracingInstanceUploadBuffer);
 
-	const FRayTracingSceneInitializer2& SceneInitializer = RayTracingSceneRHI->GetInitializer();
-
 	const int32 NumSceneInstances = Instances.Num();
 	const int32 MinBatchSize = 128;
 	ParallelFor(
@@ -178,8 +237,7 @@ void FillRayTracingInstanceUploadBuffer(
 			InstanceGeometryIndices,
 			BaseUploadBufferOffsets,
 			BaseInstancePrefixSum,
-			PreViewTranslation,
-			&SceneInitializer
+			PreViewTranslation
 		](int32 SceneInstanceIndex)
 		{
 			const FRayTracingGeometryInstance& SceneInstance = Instances[SceneInstanceIndex];
diff --git a/Engine/Source/Runtime/Renderer/Private/RayTracing/RayTracingScene.cpp b/Engine/Source/Runtime/Renderer/Private/RayTracing/RayTracingScene.cpp
--- a/Engine/Source/Runtime/Renderer/Private/RayTracing/RayTracingScene.cpp
+++ b/Engine/Source/Runtime/Renderer/Private/RayTracing/RayTracingScene.cpp
@@ -33,6 +33,10 @@ END_SHADER_PARAMETER_STRUCT()
 
 const FRayTracingScene::FInstanceHandle FRayTracingScene::INVALID_INSTANCE_HANDLE = FInstanceHandle();
 
+// Round up buffer sizes to some multiple to avoid pathological growth reallocations.
+static constexpr uint32 AllocationGranularity = 8 * 1024;
+static constexpr uint64 BufferAllocationGranularity = 16 * 1024 * 1024;
+
 FRayTracingScene::FRayTracingScene()
 {
 	const uint8 NumLayers = uint8(ERayTracingSceneLayer::NUM);
@@ -46,18 +50,16 @@ FRayTracingScene::~FRayTracingScene()
 
 void FRayTracingScene::BuildInitializationData()
 {
-	const ERayTracingAccelerationStructureFlags BuildFlags = CVarRayTracingSceneBuildMode.GetValueOnRenderThread()
-		? ERayTracingAccelerationStructureFlags::FastTrace
-		: ERayTracingAccelerationStructureFlags::FastBuild;
-
 	const uint8 NumLayers = uint8(ERayTracingSceneLayer::NUM);
 
 	for (uint32 LayerIndex = 0; LayerIndex < NumLayers; ++LayerIndex)
 	{
 		FLayer& Layer = Layers[LayerIndex];
 
-		Layer.InitializationData = CreateRayTracingSceneWithGeometryInstances(Layer.Instances, BuildFlags);
+		Layer.InitializationData = BuildRayTracingSceneInitializationData(Layer.Instances);
 	}
+
+	bInitializationDataBuilt = true;
 }
 
 void FRayTracingScene::InitPreViewTranslation(const FViewMatrices& ViewMatrices)
@@ -70,9 +72,14 @@ void FRayTracingScene::Create(FRDGBuilder& GraphBuilder, const FViewInfo& View,
 	TRACE_CPUPROFILER_EVENT_SCOPE(FRayTracingScene::Create);
 	QUICK_SCOPE_CYCLE_COUNTER(STAT_RayTracingScene_Create);
 
-	// Round up buffer sizes to some multiple to avoid pathological growth reallocations.
-	static constexpr uint32 AllocationGranularity = 8 * 1024;
-	static constexpr uint64 BufferAllocationGranularity = 16 * 1024 * 1024;
+	const ERayTracingAccelerationStructureFlags BuildFlags = CVarRayTracingSceneBuildMode.GetValueOnRenderThread()
+		? ERayTracingAccelerationStructureFlags::FastTrace
+		: ERayTracingAccelerationStructureFlags::FastBuild;
+
+	if (!bInitializationDataBuilt)
+	{
+		BuildInitializationData();
+	}
 
 	bUsedThisFrame = true;
 
@@ -84,15 +91,18 @@ void FRayTracingScene::Create(FRDGBuilder& GraphBuilder, const FViewInfo& View,
 	{
 		FLayer& Layer = Layers[LayerIndex];
 
-		if (!Layer.InitializationData.Scene.IsValid())
 		{
-			BuildInitializationData();
-		}
+			TArray<TRefCountPtr<FRHIRayTracingGeometry>> ReferencedGeometries;
+			TArray<FRHIRayTracingGeometry*> PerInstanceGeometries;
 
-		checkf(Layer.InitializationData.Scene.IsValid(),
-			TEXT("Ray tracing scene RHI object is expected to have been created by BuildInitializationData()"));
+			FRayTracingSceneInitializer2 Initializer;
+			Initializer.DebugName = FName(TEXT("FRayTracingScene"));
+			Initializer.NumNativeInstances = Layer.InitializationData.NumNativeGPUSceneInstances + Layer.InitializationData.NumNativeCPUInstances;
+			Initializer.NumTotalSegments = Layer.InitializationData.TotalNumSegments;
+			Initializer.BuildFlags = BuildFlags;
 
-		Layer.RayTracingSceneRHI = Layer.InitializationData.Scene;
+			Layer.RayTracingSceneRHI = RHICreateRayTracingScene(MoveTemp(Initializer));
+		}
 
 		const uint32 NumNativeInstances = Layer.InitializationData.NumNativeGPUSceneInstances + Layer.InitializationData.NumNativeCPUInstances;
 		const uint32 NumNativeInstancesAligned = FMath::DivideAndRoundUp(FMath::Max(NumNativeInstances, 1U), AllocationGranularity) * AllocationGranularity;
@@ -111,7 +121,7 @@ void FRayTracingScene::Create(FRDGBuilder& GraphBuilder, const FViewInfo& View,
 
 			Layer.RayTracingScenePooledBuffer = AllocatePooledBuffer(Desc, TEXT("FRayTracingScene::SceneBuffer"));
 		}
-		
+
 		Layer.RayTracingSceneBufferRDG = GraphBuilder.RegisterExternalBuffer(Layer.RayTracingScenePooledBuffer);
 		Layer.RayTracingSceneBufferSRV = GraphBuilder.CreateSRV(FRDGBufferSRVDesc(Layer.RayTracingSceneBufferRDG, Layer.RayTracingSceneRHI, 0));
 
@@ -494,7 +504,7 @@ FRDGBufferSRVRef FRayTracingScene::GetLayerView(ERayTracingSceneLayer Layer) con
 uint32 FRayTracingScene::GetNumNativeInstances(ERayTracingSceneLayer InLayer) const
 {
 	const FLayer& Layer = Layers[uint8(InLayer)];
-	checkf(Layer.InitializationData.Scene.IsValid(), TEXT("Must call BuildInitializationData() or Create() before using GetNumNativeInstances()."));
+	checkf(bInitializationDataBuilt, TEXT("Must call BuildInitializationData() or Create() before using GetNumNativeInstances()."));
 	return Layer.InitializationData.NumNativeCPUInstances + Layer.InitializationData.NumNativeGPUSceneInstances;
 }
 
@@ -634,6 +644,7 @@ void FRayTracingScene::EndFrame()
 	}
 
 	bUsedThisFrame = false;
+	bInitializationDataBuilt = false;
 }
 
 void FRayTracingScene::ReleaseReadbackBuffers()
diff --git a/Engine/Source/Runtime/Renderer/Private/RayTracing/RayTracingScene.h b/Engine/Source/Runtime/Renderer/Private/RayTracing/RayTracingScene.h
--- a/Engine/Source/Runtime/Renderer/Private/RayTracing/RayTracingScene.h
+++ b/Engine/Source/Runtime/Renderer/Private/RayTracing/RayTracingScene.h
@@ -180,7 +180,7 @@ private:
 
 	struct FLayer
 	{
-		FRayTracingSceneWithGeometryInstances InitializationData;
+		FRayTracingSceneInitializationData InitializationData;
 
 		FRayTracingSceneRHIRef RayTracingSceneRHI;
 
@@ -218,6 +218,7 @@ private:
 
 	bool bInstanceDebugDataEnabled = false;
 
+	bool bInitializationDataBuilt = false;
 	bool bUsedThisFrame = false;
 
 #if STATS
diff --git a/Engine/Source/Runtime/Renderer/Private/Tests/RayTracingTestbed.cpp b/Engine/Source/Runtime/Renderer/Private/Tests/RayTracingTestbed.cpp
--- a/Engine/Source/Runtime/Renderer/Private/Tests/RayTracingTestbed.cpp
+++ b/Engine/Source/Runtime/Renderer/Private/Tests/RayTracingTestbed.cpp
@@ -177,12 +177,23 @@ bool RunRayTracingTestbed_RenderThread(const FString& Parameters)
 	Instances[0].Transforms = MakeArrayView(&FMatrix::Identity, 1);
 	Instances[0].InstanceContributionToHitGroupIndex = 0;
 
-	FRayTracingSceneWithGeometryInstances RayTracingScene = CreateRayTracingSceneWithGeometryInstances(Instances);
+	const FRayTracingSceneInitializationData RayTracingSceneInitializationData = BuildRayTracingSceneInitializationData(Instances);
+	const ERayTracingAccelerationStructureFlags SceneBuildFlags = ERayTracingAccelerationStructureFlags::FastTrace;
 
-	const FRayTracingSceneInitializer2& SceneInitializer = RayTracingScene.Scene->GetInitializer();
+	FRayTracingSceneRHIRef RayTracingSceneRHI;
+	{
+		FRayTracingSceneInitializer2 Initializer;
+		Initializer.DebugName = FName(TEXT("FRayTracingScene"));
+		Initializer.NumNativeInstances = RayTracingSceneInitializationData.NumNativeGPUSceneInstances + RayTracingSceneInitializationData.NumNativeCPUInstances;
+		Initializer.NumTotalSegments = RayTracingSceneInitializationData.TotalNumSegments;
+		Initializer.BuildFlags = SceneBuildFlags;
+
+		RayTracingSceneRHI = RHICreateRayTracingScene(MoveTemp(Initializer));
+	}
+
+	const FRayTracingSceneInitializer2& SceneInitializer = RayTracingSceneRHI->GetInitializer();
 
-	ERayTracingAccelerationStructureFlags SceneBuildFlags = ERayTracingAccelerationStructureFlags::FastTrace;
-	FRayTracingAccelerationStructureSize SceneSizeInfo = RHICalcRayTracingSceneSize(1, SceneBuildFlags);
+	FRayTracingAccelerationStructureSize SceneSizeInfo = RHICalcRayTracingSceneSize(NumInstances, SceneBuildFlags);
 	FRHIResourceCreateInfo SceneBufferCreateInfo(TEXT("RayTracingTestBedSceneBuffer"));
 	FBufferRHIRef SceneBuffer = RHICmdList.CreateBuffer(
 		uint32(SceneSizeInfo.ResultSize),
@@ -210,7 +221,7 @@ bool RunRayTracingTestbed_RenderThread(const FString& Parameters)
 		InstanceUploadSRV = RHICmdList.CreateShaderResourceView(InstanceUploadBuffer);
 	}
 
-	const uint32 TransformUploadBufferSize = RayTracingScene.NumNativeCPUInstances * 3 * sizeof(FVector4f);
+	const uint32 TransformUploadBufferSize = RayTracingSceneInitializationData.NumNativeCPUInstances * 3 * sizeof(FVector4f);
 	FBufferRHIRef TransformUploadBuffer;
 	FShaderResourceViewRHIRef TransformUploadSRV;
 	{
@@ -223,16 +234,16 @@ bool RunRayTracingTestbed_RenderThread(const FString& Parameters)
 		FRayTracingInstanceDescriptorInput* InstanceUploadData = (FRayTracingInstanceDescriptorInput*)RHICmdList.LockBuffer(InstanceUploadBuffer, 0, InstanceUploadBufferSize, RLM_WriteOnly);
 		FVector4f* TransformUploadData = (FVector4f*)RHICmdList.LockBuffer(TransformUploadBuffer, 0, TransformUploadBufferSize, RLM_WriteOnly);
 		FillRayTracingInstanceUploadBuffer(
-			RayTracingScene.Scene,
+			RayTracingSceneRHI,
 			FVector::ZeroVector,
 			Instances,
-			RayTracingScene.InstanceGeometryIndices,
-			RayTracingScene.BaseUploadBufferOffsets,
-			RayTracingScene.BaseInstancePrefixSum,
-			RayTracingScene.NumNativeGPUSceneInstances,
-			RayTracingScene.NumNativeCPUInstances,
+			RayTracingSceneInitializationData.InstanceGeometryIndices,
+			RayTracingSceneInitializationData.BaseUploadBufferOffsets,
+			RayTracingSceneInitializationData.BaseInstancePrefixSum,
+			RayTracingSceneInitializationData.NumNativeGPUSceneInstances,
+			RayTracingSceneInitializationData.NumNativeCPUInstances,
 			MakeArrayView(InstanceUploadData, SceneInitializer.NumNativeInstances),
-			MakeArrayView(TransformUploadData, RayTracingScene.NumNativeCPUInstances * 3));
+			MakeArrayView(TransformUploadData, RayTracingSceneInitializationData.NumNativeCPUInstances * 3));
 		RHICmdList.UnlockBuffer(TransformUploadBuffer);
 		RHICmdList.UnlockBuffer(InstanceUploadBuffer);
 	}
@@ -243,9 +254,9 @@ bool RunRayTracingTestbed_RenderThread(const FString& Parameters)
 			AccelerationStructureAddressesBuffer.Buffer,
 			GPUIndex,
 			0,
-			RayTracingScene.ReferencedGeometries.Num() * sizeof(FRayTracingAccelerationStructureAddress), RLM_WriteOnly);
+			RayTracingSceneInitializationData.ReferencedGeometries.Num() * sizeof(FRayTracingAccelerationStructureAddress), RLM_WriteOnly);
 
-		const TArrayView<FRHIRayTracingGeometry*> ReferencedGeometries = RHICmdList.AllocArray(MakeConstArrayView(RayTracingScene.ReferencedGeometries));
+		const TArrayView<FRHIRayTracingGeometry*> ReferencedGeometries = RHICmdList.AllocArray(MakeConstArrayView(RayTracingSceneInitializationData.ReferencedGeometries));
 
 		RHICmdList.EnqueueLambda([AddressesPtr, ReferencedGeometries, GPUIndex](FRHICommandListBase&)
 		{
@@ -266,13 +277,13 @@ bool RunRayTracingTestbed_RenderThread(const FString& Parameters)
 		InstanceUploadSRV,
 		AccelerationStructureAddressesBuffer.SRV,
 		TransformUploadSRV,
-		RayTracingScene.NumNativeGPUSceneInstances,
-		RayTracingScene.NumNativeCPUInstances,
+		RayTracingSceneInitializationData.NumNativeGPUSceneInstances,
+		RayTracingSceneInitializationData.NumNativeCPUInstances,
 		nullptr,
 		nullptr,
 		nullptr);
 
-	RHICmdList.BindAccelerationStructureMemory(RayTracingScene.Scene, SceneBuffer, 0);
+	RHICmdList.BindAccelerationStructureMemory(RayTracingSceneRHI, SceneBuffer, 0);
 
 	RHICmdList.BuildAccelerationStructure(Geometry);
 
@@ -280,19 +291,19 @@ bool RunRayTracingTestbed_RenderThread(const FString& Parameters)
 	// RHICmdList.Transition(FRHITransitionInfo(Geometry.GetReference(), ERHIAccess::BVHWrite, ERHIAccess::BVHRead));
 
 	FRayTracingSceneBuildParams BuildParams;
-	BuildParams.Scene = RayTracingScene.Scene;
+	BuildParams.Scene = RayTracingSceneRHI;
 	BuildParams.ScratchBuffer = ScratchBuffer;
 	BuildParams.ScratchBufferOffset = 0;
 	BuildParams.InstanceBuffer = InstanceBuffer.Buffer;
 	BuildParams.InstanceBufferOffset = 0;
-	BuildParams.ReferencedGeometries = RayTracingScene.ReferencedGeometries;
-	BuildParams.PerInstanceGeometries = RayTracingScene.PerInstanceGeometries;
+	BuildParams.ReferencedGeometries = RayTracingSceneInitializationData.ReferencedGeometries;
+	BuildParams.PerInstanceGeometries = RayTracingSceneInitializationData.PerInstanceGeometries;
 
 	RHICmdList.BuildAccelerationStructure(BuildParams);
 
-	RHICmdList.Transition(FRHITransitionInfo(RayTracingScene.Scene.GetReference(), ERHIAccess::BVHWrite, ERHIAccess::BVHRead));
+	RHICmdList.Transition(FRHITransitionInfo(RayTracingSceneRHI.GetReference(), ERHIAccess::BVHWrite, ERHIAccess::BVHRead));
 
-	FShaderResourceViewInitializer RayTracingSceneViewInitializer(SceneBuffer, RayTracingScene.Scene, 0);
+	FShaderResourceViewInitializer RayTracingSceneViewInitializer(SceneBuffer, RayTracingSceneRHI, 0);
 	FShaderResourceViewRHIRef RayTracingSceneView = RHICmdList.CreateShaderResourceView(RayTracingSceneViewInitializer);
 
 	DispatchBasicOcclusionRays(RHICmdList, RayTracingSceneView, Geometry, RayBufferView, OcclusionResultBufferView, NumRays);
diff --git a/Engine/Source/Runtime/Renderer/Public/RayTracingInstanceBufferUtil.h b/Engine/Source/Runtime/Renderer/Public/RayTracingInstanceBufferUtil.h
--- a/Engine/Source/Runtime/Renderer/Public/RayTracingInstanceBufferUtil.h
+++ b/Engine/Source/Runtime/Renderer/Public/RayTracingInstanceBufferUtil.h
@@ -43,6 +43,28 @@ struct FRayTracingGPUInstance
 	uint32 DescBufferOffset;
 };
 
+struct FRayTracingSceneInitializationData
+{
+	uint32 NumNativeGPUSceneInstances;
+	uint32 NumNativeCPUInstances;
+	uint32 TotalNumSegments;
+
+	// index of each instance geometry in FRayTracingSceneRHIRef ReferencedGeometries
+	TArray<uint32> InstanceGeometryIndices;
+	// base offset of each instance entries in the instance upload buffer
+	TArray<uint32> BaseUploadBufferOffsets;
+	// prefix sum of `Instance.NumTransforms` for all instances in this scene
+	TArray<uint32> BaseInstancePrefixSum;
+
+	// Unique list of geometries referenced by all instances in this scene.
+	// Any referenced geometry is kept alive while the scene is alive.
+	TArray<FRHIRayTracingGeometry*> PerInstanceGeometries;
+	// One entry per instance
+	TArray<FRHIRayTracingGeometry*> ReferencedGeometries;
+};
+
+RENDERER_API FRayTracingSceneInitializationData BuildRayTracingSceneInitializationData(TConstArrayView<FRayTracingGeometryInstance> Instances);
+
 struct FRayTracingSceneWithGeometryInstances
 {
 	FRayTracingSceneRHIRef Scene;
@@ -85,10 +107,6 @@ RENDERER_API FRayTracingSceneWithGeometryInstances CreateRayTracingSceneWithGeom
 	uint32 NumCallableShaderSlots = 0,
 	ERayTracingAccelerationStructureFlags BuildFlags = ERayTracingAccelerationStructureFlags::FastTrace);
 
-RENDERER_API FRayTracingSceneWithGeometryInstances CreateRayTracingSceneWithGeometryInstances(
-	TConstArrayView<FRayTracingGeometryInstance> Instances,
-	ERayTracingAccelerationStructureFlags BuildFlags = ERayTracingAccelerationStructureFlags::FastTrace);
-
 // Helper function to fill upload buffers required by BuildRayTracingInstanceBuffer with instance descriptors
 // Transforms of CPU instances are copied to OutTransformData
 RENDERER_API void FillRayTracingInstanceUploadBuffer(
