[GPULM] Set MaxVertices correctly for landscape segments otherwise we create incorrect SRV and crash.

#rnx

diff --git a/Engine/Plugins/Experimental/GPULightmass/Source/GPULightmass/Private/LightmapRenderer.cpp b/Engine/Plugins/Experimental/GPULightmass/Source/GPULightmass/Private/LightmapRenderer.cpp
--- a/Engine/Plugins/Experimental/GPULightmass/Source/GPULightmass/Private/LightmapRenderer.cpp
+++ b/Engine/Plugins/Experimental/GPULightmass/Source/GPULightmass/Private/LightmapRenderer.cpp
@@ -1020,7 +1020,7 @@ bool FSceneRenderState::SetupRayTracingScene(FRDGBuilder& GraphBuilder, FSceneUn
 							Segment.VertexBuffer = nullptr;
 							Segment.VertexBufferStride = sizeof(FVector3f);
 							Segment.VertexBufferElementType = VET_Float3;
-							Segment.MaxVertices = NumPrimitives * 3;
+							Segment.MaxVertices = FMath::Square(Landscape.SubsectionSizeVerts);
 							Segment.NumPrimitives = NumPrimitives;
 							GeometryInitializer.Segments.Add(Segment);
 
