[Engine]
* Removed SkyAtmosphereCommonData.h include from SceneManagement.h

diff --git a/Engine/Plugins/Experimental/GPULightmass/Source/GPULightmass/Private/Scene/Lights.cpp b/Engine/Plugins/Experimental/GPULightmass/Source/GPULightmass/Private/Scene/Lights.cpp
--- a/Engine/Plugins/Experimental/GPULightmass/Source/GPULightmass/Private/Scene/Lights.cpp
+++ b/Engine/Plugins/Experimental/GPULightmass/Source/GPULightmass/Private/Scene/Lights.cpp
@@ -4,6 +4,7 @@
 
 #include "Engine/TextureLightProfile.h"
 #include "LightmapRayTracing.h"
+#include "Rendering/SkyAtmosphereCommonData.h"
 #include "Scene.h"
 #include "RenderGraphBuilder.h"
 #include "ReflectionEnvironment.h"
