Added missing UnlockBuffer(InstanceGroupUploadBuffer).

diff --git a/Engine/Plugins/Experimental/GPULightmass/Source/GPULightmass/Private/LightmapRenderer.cpp b/Engine/Plugins/Experimental/GPULightmass/Source/GPULightmass/Private/LightmapRenderer.cpp
--- a/Engine/Plugins/Experimental/GPULightmass/Source/GPULightmass/Private/LightmapRenderer.cpp
+++ b/Engine/Plugins/Experimental/GPULightmass/Source/GPULightmass/Private/LightmapRenderer.cpp
@@ -1239,6 +1239,7 @@ bool FSceneRenderState::SetupRayTracingScene(FRDGBuilder& GraphBuilder, FSceneUn
 					MakeArrayView(TransformUploadData, RayTracingSceneInitializationData.NumNativeCPUInstances * 3));
 				RHICmdList.UnlockBuffer(TransformUploadBuffer);
 				RHICmdList.UnlockBuffer(InstanceUploadBuffer);
+				RHICmdList.UnlockBuffer(InstanceGroupUploadBuffer);
 			}
 
 			for (uint32 GPUIndex : RHICmdList.GetGPUMask())
diff --git a/Engine/Source/Runtime/Renderer/Private/Tests/RayTracingTestbed.cpp b/Engine/Source/Runtime/Renderer/Private/Tests/RayTracingTestbed.cpp
--- a/Engine/Source/Runtime/Renderer/Private/Tests/RayTracingTestbed.cpp
+++ b/Engine/Source/Runtime/Renderer/Private/Tests/RayTracingTestbed.cpp
@@ -261,6 +261,7 @@ bool RunRayTracingTestbed_RenderThread(const FString& Parameters)
 			MakeArrayView(TransformUploadData, RayTracingSceneInitializationData.NumNativeCPUInstances * 3));
 		RHICmdList.UnlockBuffer(TransformUploadBuffer);
 		RHICmdList.UnlockBuffer(InstanceUploadBuffer);
+		RHICmdList.UnlockBuffer(InstanceGroupUploadBuffer);
 	}
 
 	for (uint32 GPUIndex : RHICmdList.GetGPUMask())
