Chnage AGPULightmassSettingsActor to be a main world actor type. This fixes issues when the GPU Lightmass plugin is enabled and level instance may contain AGPULightmassSettingsActors. Specifically, breaking a level instance does not filter the actor out can result in multiple of them being moved into the parent world.

diff --git a/Engine/Plugins/Experimental/GPULightmass/Source/GPULightmass/Public/GPULightmassSettings.h b/Engine/Plugins/Experimental/GPULightmass/Source/GPULightmass/Public/GPULightmassSettings.h
--- a/Engine/Plugins/Experimental/GPULightmass/Source/GPULightmass/Public/GPULightmassSettings.h
+++ b/Engine/Plugins/Experimental/GPULightmass/Source/GPULightmass/Public/GPULightmassSettings.h
@@ -144,6 +144,10 @@ class GPULIGHTMASS_API AGPULightmassSettingsActor : public AInfo
 {
 	GENERATED_UCLASS_BODY()
 
+#if WITH_EDITOR
+	virtual bool ActorTypeIsMainWorldOnly() const override { return true; }
+#endif
+
 public:
 	UPROPERTY(EditAnywhere, BlueprintReadWrite, Category = Settings)
 	TObjectPtr<UGPULightmassSettings> Settings;
