Merging RHI CreateBuffer refactor to Main.
- Buffers are now created using FRHIBufferInitializer objects via RHICreateBufferInitializer using the new FRHIBufferCreateDesc struct.
- FRHIBufferCreateDesc specifies if the buffer should be initialized with default data, zeroed data, filled with resource array data, or use data provided to FRHIBufferInitializer.
- FRHIBufferInitializer objects can carry platform-agnostic pointers for initializing the new buffer. This allows many unified memory platforms to remove extra copies when creating buffers.
- FRHIBuffer objects are created via FRHIBufferInitializer::Finalize() which internally will queue uploads of pending data, depending on the platform.
- Converted a number of locations that did CreateBuffer->LockBuffer->UnlockBuffer to use the Buffer Initializers.
- Added unit tests for Buffer Initializers.

diff --git a/Engine/Plugins/2D/Paper2D/Source/Paper2D/Private/PaperSpriteVertexBuffer.cpp b/Engine/Plugins/2D/Paper2D/Source/Paper2D/Private/PaperSpriteVertexBuffer.cpp
--- a/Engine/Plugins/2D/Paper2D/Source/Paper2D/Private/PaperSpriteVertexBuffer.cpp
+++ b/Engine/Plugins/2D/Paper2D/Source/Paper2D/Private/PaperSpriteVertexBuffer.cpp
@@ -27,8 +27,12 @@ void FPaperSpriteVertexBuffer::CreateBuffers(FRHICommandListBase& RHICmdList, in
 	uint32 PositionSize = NumAllocatedVertices * sizeof(FVector3f);
 	// create vertex buffer
 	{
-		FRHIResourceCreateInfo CreateInfo(TEXT("PaperSpritePositionBuffer"));
-		PositionBuffer.VertexBufferRHI = RHICmdList.CreateVertexBuffer(PositionSize, Usage, CreateInfo);
+		const FRHIBufferCreateDesc CreateDesc =
+			FRHIBufferCreateDesc::CreateVertex(TEXT("PaperSpritePositionBuffer"), PositionSize)
+			.AddUsage(Usage)
+			.DetermineInitialState();
+
+		PositionBuffer.VertexBufferRHI = RHICmdList.CreateBuffer(CreateDesc);
 		if (RHISupportsManualVertexFetch(GMaxRHIShaderPlatform))
 		{
 			PositionBufferSRV = RHICmdList.CreateShaderResourceView(
@@ -43,8 +47,12 @@ void FPaperSpriteVertexBuffer::CreateBuffers(FRHICommandListBase& RHICmdList, in
 	uint32 TangentSize = NumAllocatedVertices * 2 * sizeof(FPackedNormal);
 	// create vertex buffer
 	{
-		FRHIResourceCreateInfo CreateInfo(TEXT("PaperSpriteTangentBuffer"));
-		TangentBuffer.VertexBufferRHI = RHICmdList.CreateVertexBuffer(TangentSize, Usage, CreateInfo);
+		const FRHIBufferCreateDesc CreateDesc =
+			FRHIBufferCreateDesc::CreateVertex(TEXT("PaperSpriteTangentBuffer"), TangentSize)
+			.AddUsage(Usage)
+			.DetermineInitialState();
+
+		TangentBuffer.VertexBufferRHI = RHICmdList.CreateBuffer(CreateDesc);
 		if (RHISupportsManualVertexFetch(GMaxRHIShaderPlatform))
 		{
 			TangentBufferSRV = RHICmdList.CreateShaderResourceView(
@@ -58,8 +66,12 @@ void FPaperSpriteVertexBuffer::CreateBuffers(FRHICommandListBase& RHICmdList, in
 	uint32 TexCoordSize = NumAllocatedVertices * sizeof(FVector2f);
 	// create vertex buffer
 	{
-		FRHIResourceCreateInfo CreateInfo(TEXT("PaperSpriteTexCoordBuffer"));
-		TexCoordBuffer.VertexBufferRHI = RHICmdList.CreateVertexBuffer(TexCoordSize, Usage, CreateInfo);
+		const FRHIBufferCreateDesc CreateDesc =
+			FRHIBufferCreateDesc::CreateVertex(TEXT("PaperSpriteTexCoordBuffer"), TexCoordSize)
+			.AddUsage(Usage)
+			.DetermineInitialState();
+
+		TexCoordBuffer.VertexBufferRHI = RHICmdList.CreateBuffer(CreateDesc);
 		if (RHISupportsManualVertexFetch(GMaxRHIShaderPlatform))
 		{
 			TexCoordBufferSRV = RHICmdList.CreateShaderResourceView(
@@ -73,8 +85,12 @@ void FPaperSpriteVertexBuffer::CreateBuffers(FRHICommandListBase& RHICmdList, in
 	uint32 ColorSize = NumAllocatedVertices * sizeof(FColor);
 	// create vertex buffer
 	{
-		FRHIResourceCreateInfo CreateInfo(TEXT("PaperSpriteColorBuffer"));
-		ColorBuffer.VertexBufferRHI = RHICmdList.CreateVertexBuffer(ColorSize, Usage, CreateInfo);
+		const FRHIBufferCreateDesc CreateDesc =
+			FRHIBufferCreateDesc::CreateVertex(TEXT("PaperSpriteColorBuffer"), ColorSize)
+			.AddUsage(Usage)
+			.DetermineInitialState();
+
+		ColorBuffer.VertexBufferRHI = RHICmdList.CreateBuffer(CreateDesc);
 		if (RHISupportsManualVertexFetch(GMaxRHIShaderPlatform))
 		{
 			ColorBufferSRV = RHICmdList.CreateShaderResourceView(
@@ -87,8 +103,12 @@ void FPaperSpriteVertexBuffer::CreateBuffers(FRHICommandListBase& RHICmdList, in
 
 	//Create Index Buffer
 	{
-		FRHIResourceCreateInfo CreateInfo(TEXT("PaperSpriteIndexBuffer"));
-		IndexBuffer.IndexBufferRHI = RHICmdList.CreateIndexBuffer(sizeof(uint32), Vertices.Num() * sizeof(uint32), Usage, CreateInfo);
+		const FRHIBufferCreateDesc CreateDesc =
+			FRHIBufferCreateDesc::CreateIndex<uint32>(TEXT("PaperSpriteIndexBuffer"), Vertices.Num())
+			.AddUsage(Usage)
+			.DetermineInitialState();
+
+		IndexBuffer.IndexBufferRHI = RHICmdList.CreateBuffer(CreateDesc);
 	}
 }
 
diff --git a/Engine/Plugins/Animation/MLDeformer/MLDeformerFramework/Source/MLDeformerFramework/Private/MLDeformerModel.cpp b/Engine/Plugins/Animation/MLDeformer/MLDeformerFramework/Source/MLDeformerFramework/Private/MLDeformerModel.cpp
--- a/Engine/Plugins/Animation/MLDeformer/MLDeformerFramework/Source/MLDeformerFramework/Private/MLDeformerModel.cpp
+++ b/Engine/Plugins/Animation/MLDeformer/MLDeformerFramework/Source/MLDeformerFramework/Private/MLDeformerModel.cpp
@@ -26,17 +26,21 @@ namespace UE::MLDeformer
 	{
 		if (VertexMap.Num() > 0)
 		{
-			FRHIResourceCreateInfo CreateInfo(TEXT("UMLDeformerModel::FVertexMapBuffer"));
+			const FRHIBufferCreateDesc CreateDesc =
+				FRHIBufferCreateDesc::CreateVertex<uint32>(TEXT("UMLDeformerModel::FVertexMapBuffer"), VertexMap.Num())
+				.AddUsage(EBufferUsageFlags::Static | EBufferUsageFlags::ShaderResource)
+				.DetermineInitialState()
+				.SetInitActionInitializer();
+
+			TRHIBufferInitializer<uint32> InitialData = RHICmdList.CreateBufferInitializer(CreateDesc);
 
-			VertexBufferRHI = RHICmdList.CreateVertexBuffer(VertexMap.Num() * sizeof(uint32), BUF_Static | BUF_ShaderResource, CreateInfo);
-			uint32* Data = reinterpret_cast<uint32*>(RHICmdList.LockBuffer(VertexBufferRHI, 0, VertexMap.Num() * sizeof(uint32), RLM_WriteOnly));
 			for (int32 Index = 0; Index < VertexMap.Num(); ++Index)
 			{
-				Data[Index] = static_cast<uint32>(VertexMap[Index]);
+				InitialData[Index] = static_cast<uint32>(VertexMap[Index]);
 			}
-			RHICmdList.UnlockBuffer(VertexBufferRHI);
 			VertexMap.Empty();
 
+			VertexBufferRHI = InitialData.Finalize();
 			ShaderResourceViewRHI = RHICmdList.CreateShaderResourceView(
 				VertexBufferRHI, 
 				FRHIViewDesc::CreateBufferSRV()
diff --git a/Engine/Plugins/Enterprise/LidarPointCloud/Source/LidarPointCloudRuntime/Private/Rendering/LidarPointCloudRenderBuffers.cpp b/Engine/Plugins/Enterprise/LidarPointCloud/Source/LidarPointCloudRuntime/Private/Rendering/LidarPointCloudRenderBuffers.cpp
--- a/Engine/Plugins/Enterprise/LidarPointCloud/Source/LidarPointCloudRuntime/Private/Rendering/LidarPointCloudRenderBuffers.cpp
+++ b/Engine/Plugins/Enterprise/LidarPointCloud/Source/LidarPointCloudRuntime/Private/Rendering/LidarPointCloudRenderBuffers.cpp
@@ -44,13 +44,16 @@ void FLidarPointCloudIndexBuffer::Resize(const uint32 & RequestedCapacity)
 
 void FLidarPointCloudIndexBuffer::InitRHI(FRHICommandListBase& RHICmdList)
 {
-	FRHIResourceCreateInfo CreateInfo(TEXT("FLidarPointCloudIndexBuffer"));
-	const uint32 Size = Capacity * 7 * sizeof(uint32);
 	PointOffset = Capacity * 6;
 
-	IndexBufferRHI = RHICmdList.CreateBuffer(Size, BUF_Dynamic | BUF_IndexBuffer, sizeof(uint32), ERHIAccess::VertexOrIndexBuffer, CreateInfo);
+	const FRHIBufferCreateDesc CreateDesc =
+		FRHIBufferCreateDesc::CreateIndex<uint32>(TEXT("FLidarPointCloudIndexBuffer"), Capacity * 7)
+		.AddUsage(EBufferUsageFlags::Dynamic)
+		.SetInitialState(ERHIAccess::VertexOrIndexBuffer)
+		.SetInitActionInitializer();
+
+	TRHIBufferInitializer<uint32> Data = RHICmdList.CreateBufferInitializer(CreateDesc);
 
-	uint32* Data = (uint32*)RHICmdList.LockBuffer(IndexBufferRHI, 0, Size, RLM_WriteOnly);
 	for (uint32 i = 0, idx = 0; i < Capacity; i++)
 	{
 		const uint32 v = i * 4;
@@ -66,7 +69,8 @@ void FLidarPointCloudIndexBuffer::InitRHI(FRHICommandListBase& RHICmdList)
 		// Points
 		Data[PointOffset + i] = v;
 	}
-	RHICmdList.UnlockBuffer(IndexBufferRHI);
+
+	IndexBufferRHI = Data.Finalize();
 }
 
 //////////////////////////////////////////////////////////// Structured Buffer
@@ -90,8 +94,12 @@ void FLidarPointCloudRenderBuffer::Resize(const uint32& RequestedCapacity)
 
 void FLidarPointCloudRenderBuffer::InitRHI(FRHICommandListBase& RHICmdList)
 {
-	FRHIResourceCreateInfo CreateInfo(TEXT("FLidarPointCloudRenderBuffer"));
-	Buffer = RHICmdList.CreateVertexBuffer(sizeof(uint32) * Capacity, BUF_ShaderResource | BUF_Dynamic, CreateInfo);
+	const FRHIBufferCreateDesc CreateDesc =
+		FRHIBufferCreateDesc::CreateVertex<uint32>(TEXT("FLidarPointCloudRenderBuffer"), Capacity)
+		.AddUsage(EBufferUsageFlags::ShaderResource | EBufferUsageFlags::Dynamic)
+		.DetermineInitialState();
+
+	Buffer = RHICmdList.CreateBuffer(CreateDesc);
 	SRV = RHICmdList.CreateShaderResourceView(
 		Buffer,
 		FRHIViewDesc::CreateBufferSRV()
@@ -334,16 +342,16 @@ void FLidarPointCloudVertexFactory::Initialize(FLidarPointCloudPoint* Data, int3
 void FLidarPointCloudVertexFactory::FPointCloudVertexBuffer::InitRHI(FRHICommandListBase& RHICmdList)
 {
 	const uint32 BufferSize = NumPoints * 4 * sizeof(FLidarPointCloudPoint);
-	FRHIResourceCreateInfo CreateInfo(TEXT("FPointCloudVertexBuffer"));
-	VertexBufferRHI = RHICmdList.CreateBuffer(BufferSize, BUF_ShaderResource | BUF_Static | BUF_VertexBuffer, sizeof(FLidarPointCloudPoint), ERHIAccess::VertexOrIndexBuffer, CreateInfo);
-	VertexBufferSRV = RHICmdList.CreateShaderResourceView(
-		VertexBufferRHI, 
-		FRHIViewDesc::CreateBufferSRV()
-			.SetType(FRHIViewDesc::EBufferType::Typed)
-			.SetFormat(PF_R32_FLOAT));
-	void* Buffer = RHICmdList.LockBuffer(VertexBufferRHI, 0, BufferSize, RLM_WriteOnly);
 
-	FLidarPointCloudPoint* Dest = (FLidarPointCloudPoint*)Buffer;
+	const FRHIBufferCreateDesc CreateDesc =
+		FRHIBufferCreateDesc::Create(TEXT("FPointCloudVertexBuffer"), BufferSize, sizeof(FLidarPointCloudPoint), EBufferUsageFlags::VertexBuffer)
+		.AddUsage(EBufferUsageFlags::ShaderResource | EBufferUsageFlags::Static)
+		.SetInitialState(ERHIAccess::VertexOrIndexBuffer)
+		.SetInitActionInitializer();
+
+	TRHIBufferInitializer<FLidarPointCloudPoint> InitialData = RHICmdList.CreateBufferInitializer(CreateDesc);
+
+	FLidarPointCloudPoint* Dest = InitialData.GetWritableData();
 	for (int32 i = 0; i < NumPoints; ++i, ++Data)
 	{
 		FMemory::Memcpy(Dest, Data, sizeof(FLidarPointCloudPoint)); Dest++;
@@ -352,7 +360,12 @@ void FLidarPointCloudVertexFactory::FPointCloudVertexBuffer::InitRHI(FRHICommand
 		FMemory::Memcpy(Dest, Data, sizeof(FLidarPointCloudPoint)); Dest++;
 	}
 
-	RHICmdList.UnlockBuffer(VertexBufferRHI);
+	VertexBufferRHI = InitialData.Finalize();
+	VertexBufferSRV = RHICmdList.CreateShaderResourceView(
+		VertexBufferRHI, 
+		FRHIViewDesc::CreateBufferSRV()
+			.SetType(FRHIViewDesc::EBufferType::Typed)
+			.SetFormat(PF_R32_FLOAT));
 }
 
 void FLidarPointCloudVertexFactory::InitRHI(FRHICommandListBase& RHICmdList)
@@ -376,11 +389,13 @@ void FLidarPointCloudVertexFactory::ReleaseRHI()
 
 void FLidarPointCloudSharedVertexFactory::FPointCloudVertexBuffer::InitRHI(FRHICommandListBase& RHICmdList)
 {
-	FRHIResourceCreateInfo CreateInfo(TEXT("FPointCloudVertexBuffer"));
-	VertexBufferRHI = RHICmdList.CreateBuffer(sizeof(FVector), BUF_Static | BUF_VertexBuffer, 0, ERHIAccess::VertexOrIndexBuffer, CreateInfo);
-	void* Buffer = RHICmdList.LockBuffer(VertexBufferRHI, 0, sizeof(FVector), RLM_WriteOnly);
-	FMemory::Memzero(Buffer, sizeof(FVector));
-	RHICmdList.UnlockBuffer(VertexBufferRHI);
+	const FRHIBufferCreateDesc CreateDesc =
+		FRHIBufferCreateDesc::CreateVertex(TEXT("FPointCloudVertexBuffer"), sizeof(FVector))
+		.AddUsage(EBufferUsageFlags::Static)
+		.SetInitialState(ERHIAccess::VertexOrIndexBuffer)
+		.SetInitActionZeroData();
+
+	VertexBufferRHI = RHICmdList.CreateBuffer(CreateDesc);
 }
 
 void FLidarPointCloudSharedVertexFactory::InitRHI(FRHICommandListBase& RHICmdList)
diff --git a/Engine/Plugins/Enterprise/LidarPointCloud/Source/LidarPointCloudRuntime/Private/Rendering/LidarPointCloudRendering.cpp b/Engine/Plugins/Enterprise/LidarPointCloud/Source/LidarPointCloudRuntime/Private/Rendering/LidarPointCloudRendering.cpp
--- a/Engine/Plugins/Enterprise/LidarPointCloud/Source/LidarPointCloudRuntime/Private/Rendering/LidarPointCloudRendering.cpp
+++ b/Engine/Plugins/Enterprise/LidarPointCloud/Source/LidarPointCloudRuntime/Private/Rendering/LidarPointCloudRendering.cpp
@@ -140,11 +140,16 @@ private:
 		{
 			const uint32 Size = DataLength * sizeof(FVector3f);
 
-			FRHIResourceCreateInfo CreateInfo(TEXT("FLidarPointCloudCollisionVertexBuffer"));
-			VertexBufferRHI = RHICmdList.CreateBuffer(Size, BUF_Static | BUF_VertexBuffer | BUF_ShaderResource, 0, ERHIAccess::VertexOrIndexBuffer | ERHIAccess::SRVMask, CreateInfo);
-			void* Buffer = RHICmdList.LockBuffer(VertexBufferRHI, 0, Size, RLM_WriteOnly);
-			FMemory::Memcpy(Buffer, Data, Size);
-			RHICmdList.UnlockBuffer(VertexBufferRHI);
+			const FRHIBufferCreateDesc CreateDesc =
+				FRHIBufferCreateDesc::CreateVertex(TEXT("FLidarPointCloudCollisionVertexBuffer"), Size)
+				.AddUsage(EBufferUsageFlags::Static | EBufferUsageFlags::ShaderResource)
+				.SetInitialState(ERHIAccess::VertexOrIndexBuffer | ERHIAccess::SRVMask)
+				.SetInitActionInitializer();
+
+			FRHIBufferInitializer InitialData = RHICmdList.CreateBufferInitializer(CreateDesc);
+			InitialData.WriteData(Data, Size);
+
+			VertexBufferRHI = InitialData.Finalize();
 		}
 	} VertexBuffer;
 	class FLidarPointCloudCollisionIndexBuffer : public FIndexBuffer
@@ -165,11 +170,16 @@ private:
 		{
 			const uint32 Size = DataLength * sizeof(uint32);
 
-			FRHIResourceCreateInfo CreateInfo(TEXT("FLidarPointCloudCollisionIndexBuffer"));
-			IndexBufferRHI = RHICmdList.CreateBuffer(Size, BUF_Static | BUF_IndexBuffer, sizeof(uint32), ERHIAccess::VertexOrIndexBuffer, CreateInfo);
-			void* Buffer = RHICmdList.LockBuffer(IndexBufferRHI, 0, Size, RLM_WriteOnly);
-			FMemory::Memcpy(Buffer, Data, Size);
-			RHICmdList.UnlockBuffer(IndexBufferRHI);
+			const FRHIBufferCreateDesc CreateDesc =
+				FRHIBufferCreateDesc::CreateIndex(TEXT("FLidarPointCloudCollisionIndexBuffer"), Size, sizeof(uint32))
+				.AddUsage(EBufferUsageFlags::Static)
+				.SetInitialState(ERHIAccess::VertexOrIndexBuffer)
+				.SetInitActionInitializer();
+
+			FRHIBufferInitializer InitialData = RHICmdList.CreateBufferInitializer(CreateDesc);
+			InitialData.WriteData(Data, Size);
+
+			IndexBufferRHI = InitialData.Finalize();
 		}
 	} IndexBuffer;
 
diff --git a/Engine/Plugins/Experimental/AVCodecs/AVCodecsCore/Source/AVCodecsCoreRHI/Private/Video/Resources/VideoResourceRHI.cpp b/Engine/Plugins/Experimental/AVCodecs/AVCodecsCore/Source/AVCodecsCoreRHI/Private/Video/Resources/VideoResourceRHI.cpp
--- a/Engine/Plugins/Experimental/AVCodecs/AVCodecsCore/Source/AVCodecsCoreRHI/Private/Video/Resources/VideoResourceRHI.cpp
+++ b/Engine/Plugins/Experimental/AVCodecs/AVCodecsCore/Source/AVCodecsCoreRHI/Private/Video/Resources/VideoResourceRHI.cpp
@@ -383,7 +383,7 @@ TSharedPtr<FVideoResourceRHI> FVideoResourceRHI::TransformResource(FVideoDescrip
 						SetShaderParametersLegacyPS(RHICmdList, ConvertShader, Source, OutputDim, YUVToRGBMatrix, UE::Color::EEncoding::sRGB, FMatrix44f::Identity, false);
 
 						// draw full size quad into render target
-						FBufferRHIRef VertexBuffer = CreateTempMediaVertexBuffer();
+						FBufferRHIRef VertexBuffer = CreateTempMediaVertexBuffer(RHICmdList);
 						RHICmdList.SetStreamSource(0, VertexBuffer, 0);
 						// set viewport to RT size
 						RHICmdList.SetViewport(0, 0, 0.0f, (float)OutputDim.X, (float)OutputDim.Y, 1.0f);
@@ -452,7 +452,7 @@ void FVideoResourceRHI::TransformResourceTo(FRHICommandListImmediate& RHICmdList
 					SetShaderParametersLegacyPS(RHICmdList, ConvertShader, Raw.Texture, OutputDim, YUVToRGBMatrix, UE::Color::EEncoding::sRGB, FMatrix44f::Identity, false);
 
 					// draw full size quad into render target
-					FBufferRHIRef VertexBuffer = CreateTempMediaVertexBuffer();
+					FBufferRHIRef VertexBuffer = CreateTempMediaVertexBuffer(RHICmdList);
 					RHICmdList.SetStreamSource(0, VertexBuffer, 0);
 					// set viewport to RT size
 					RHICmdList.SetViewport(0, 0, 0.0f, (float)OutputDim.X, (float)OutputDim.Y, 1.0f);
diff --git a/Engine/Plugins/Experimental/ChaosFlesh/Source/ChaosFleshEngine/Private/ChaosFlesh/ChaosDeformableGPUBuffers.cpp b/Engine/Plugins/Experimental/ChaosFlesh/Source/ChaosFleshEngine/Private/ChaosFlesh/ChaosDeformableGPUBuffers.cpp
--- a/Engine/Plugins/Experimental/ChaosFlesh/Source/ChaosFleshEngine/Private/ChaosFlesh/ChaosDeformableGPUBuffers.cpp
+++ b/Engine/Plugins/Experimental/ChaosFlesh/Source/ChaosFleshEngine/Private/ChaosFlesh/ChaosDeformableGPUBuffers.cpp
@@ -3,6 +3,7 @@
 #include "ChaosFlesh/ChaosDeformableGPUBuffers.h"
 #include "ChaosFlesh/ChaosDeformableTypes.h"
 
+#include "RHIResourceUtils.h"
 
 DEFINE_LOG_CATEGORY_STATIC(LogFleshDeformerBuffers, Log, All);
 
@@ -44,15 +45,12 @@ void FFloatArrayBufferWithSRV::InitRHI(FRHICommandListBase& RHICmdList)
 {
     if (Array.Num() > 0)
     {
-        FRHIResourceCreateInfo CreateInfo(TEXT("UMLDeformerModel::FFloatArrayBufferWithSRV"));
-
-        VertexBufferRHI = RHICmdList.CreateVertexBuffer(Array.Num() * sizeof(float), BUF_Static | BUF_ShaderResource, CreateInfo);
-        float* Data = reinterpret_cast<float*>(RHICmdList.LockBuffer(VertexBufferRHI, 0, Array.Num() * sizeof(float), RLM_WriteOnly));
-        for (int32 Index = 0; Index < Array.Num(); ++Index)
-        {
-            Data[Index] = Array[Index];
-        }
-		RHICmdList.UnlockBuffer(VertexBufferRHI);
+		VertexBufferRHI = UE::RHIResourceUtils::CreateVertexBufferFromArray(
+			RHICmdList,
+			TEXT("UMLDeformerModel::FFloatArrayBufferWithSRV"),
+			EBufferUsageFlags::Static | EBufferUsageFlags::ShaderResource,
+			MakeConstArrayView(Array)
+		);
 
 		if (bEmptyArray)
 		{
@@ -122,15 +120,13 @@ void FHalfArrayBufferWithSRV::InitRHI(FRHICommandListBase& RHICmdList)
 {
     if (Array.Num() > 0)
     {
-        FRHIResourceCreateInfo CreateInfo(TEXT("FHalfArrayBufferWithSRV"));
+		VertexBufferRHI = UE::RHIResourceUtils::CreateVertexBufferFromArray(
+			RHICmdList,
+			TEXT("FHalfArrayBufferWithSRV"),
+			EBufferUsageFlags::Static | EBufferUsageFlags::ShaderResource,
+			MakeConstArrayView(Array)
+		);
 
-        VertexBufferRHI = RHICmdList.CreateVertexBuffer(Array.Num() * sizeof(FFloat16), BUF_Static | BUF_ShaderResource, CreateInfo);
-        FFloat16* Data = reinterpret_cast<FFloat16*>(RHICmdList.LockBuffer(VertexBufferRHI, 0, Array.Num() * sizeof(FFloat16), RLM_WriteOnly));
-        for (int32 Index = 0; Index < Array.Num(); ++Index)
-        {
-            Data[Index] = Array[Index];
-        }
-		RHICmdList.UnlockBuffer(VertexBufferRHI);
         Array.Empty();
 
         ShaderResourceViewRHI = RHICmdList.CreateShaderResourceView(
@@ -267,61 +263,73 @@ void FIndexArrayBufferWithSRV::InitRHI(FRHICommandListBase& RHICmdList)
     if (BufferSize)
     {
         // Create the index buffer.
-        FRHIResourceCreateInfo CreateInfo(*GetFriendlyName());
-		VertexBufferRHI = RHICmdList.CreateVertexBuffer(BufferSize, BUF_Static | BUF_ShaderResource, CreateInfo);
+		const FRHIBufferCreateDesc CreateDesc =
+			FRHIBufferCreateDesc::CreateVertex(*BufferName, BufferSize)
+			.AddUsage(EBufferUsageFlags::Static | EBufferUsageFlags::ShaderResource)
+			.DetermineInitialState()
+			.SetInitActionInitializer();
 
-        // Initialize the buffer.		
-        void* Buffer = RHICmdList.LockBuffer(VertexBufferRHI, 0, BufferSize, RLM_WriteOnly);
+        // Initialize the buffer.
 		
 		if (!bForce32 && bUint8)
         {
+			TRHIBufferInitializer<uint8> Initializer = RHICmdList.CreateBufferInitializer(CreateDesc);
+
 			// 0: 33 22 11 00
 			// 1: 77 66 55 44
 			// 2: 11 10 99 88
 
-            uint8* DestIndices8Bit = reinterpret_cast<uint8*>(Buffer);
 			int32 i = 0;
 			for (; i < NumIndices; i++)
             {
-                DestIndices8Bit[i] = static_cast<uint8>(Array[i] + Offset); // add offset, sub in shader
+				Initializer[i] = static_cast<uint8>(Array[i] + Offset); // add offset, sub in shader
             }
 			// Zero tailing unused space.
 			for (; i < NumSlots * 4; i++)
 			{
-				DestIndices8Bit[i] = 0;
+				Initializer[i] = 0;
 			}
+
+			VertexBufferRHI = Initializer.Finalize();
 		}
         else if (!bForce32 && bUint16)
         {
+			TRHIBufferInitializer<uint16> Initializer = RHICmdList.CreateBufferInitializer(CreateDesc);
+
 			// 0: 1111 0000
 			// 1: 3333 2222
 			// 2: 5555 4444
 			
-			uint16* DestIndices16Bit = reinterpret_cast<uint16*>(Buffer);
 			int32 i = 0;
 			for (; i < NumIndices; i++)
 			{
-				DestIndices16Bit[i] = static_cast<uint16>(Array[i] + Offset); // add offset, sub in shader
+				Initializer[i] = static_cast<uint16>(Array[i] + Offset); // add offset, sub in shader
 			}
 			// Zero tailing unused space.
 			for (; i < NumSlots * 2; i++)
 			{
-				DestIndices16Bit[i] = 0;
+				Initializer[i] = 0;
 			}
+
+			VertexBufferRHI = Initializer.Finalize();
 		}
         else // uint32
         {
-            uint32* DestIndices32Bit = reinterpret_cast<uint32*>(Buffer);
+			TRHIBufferInitializer<uint32> Initializer = RHICmdList.CreateBufferInitializer(CreateDesc);
+
 			for (int32 i = 0; i < NumIndices; i++)
 			{
-				DestIndices32Bit[i] = static_cast<uint32>(Array[i] + Offset); // add offset, sub in shader
+				Initializer[i] = static_cast<uint32>(Array[i] + Offset); // add offset, sub in shader
 			}
+
+			VertexBufferRHI = Initializer.Finalize();
         }
 
-#if WITH_EDITOR
+		// Commented out because we should not be reading from Mapped data
+#if WITH_EDITOR && 0
 		if (Local::CVarParam.bTestUnpacking)
 		{
-			uint32* Indices32Bit = reinterpret_cast<uint32*>(Buffer);
+			uint32* Indices32Bit = reinterpret_cast<uint32*>(InitialData.GetDataForWriting());
 			for (int32 i = 0; i < NumIndices; i++)
 			{
 				int32 UnpackedInt = static_cast<int>(UnpackUInt(Indices32Bit, i, DataStride, NumIndices)) - Offset; // sub offset for unpacking
@@ -334,7 +342,6 @@ void FIndexArrayBufferWithSRV::InitRHI(FRHICommandListBase& RHICmdList)
 		}
 #endif
 
-		RHICmdList.UnlockBuffer(VertexBufferRHI);
         Array.Empty();
 
 		EPixelFormat Format = PF_R32_UINT;
diff --git a/Engine/Plugins/Experimental/ColorCorrectRegions/Source/ColorCorrectRegions/Private/ColorCorrectRegionsSceneViewExtension.cpp b/Engine/Plugins/Experimental/ColorCorrectRegions/Source/ColorCorrectRegions/Private/ColorCorrectRegionsSceneViewExtension.cpp
--- a/Engine/Plugins/Experimental/ColorCorrectRegions/Source/ColorCorrectRegions/Private/ColorCorrectRegionsSceneViewExtension.cpp
+++ b/Engine/Plugins/Experimental/ColorCorrectRegions/Source/ColorCorrectRegions/Private/ColorCorrectRegionsSceneViewExtension.cpp
@@ -285,8 +285,6 @@ namespace
 			Parameters->PostProcessOutput = SceneTextureViewportParams;
 			Parameters->View = View.ViewUniformBuffer;
 
-			FRHIResourceCreateInfo CreateInfo(TEXT("CCR_StencilIdBuffer"));
-
 			Parameters->StencilIds = GraphBuilder.CreateSRV(FRDGBufferSRVDesc(CreateStructuredBuffer(GraphBuilder, TEXT("CCR.StencilIdBuffer"), sizeof(uint32), StencilIds.Num(), &StencilIds[0], sizeof(uint32) * StencilIds.Num())));
 			Parameters->StencilIdCount = StencilIds.Num();
 
diff --git a/Engine/Plugins/Experimental/GPULightmass/Source/GPULightmass/Private/IrradianceCaching.cpp b/Engine/Plugins/Experimental/GPULightmass/Source/GPULightmass/Private/IrradianceCaching.cpp
--- a/Engine/Plugins/Experimental/GPULightmass/Source/GPULightmass/Private/IrradianceCaching.cpp
+++ b/Engine/Plugins/Experimental/GPULightmass/Source/GPULightmass/Private/IrradianceCaching.cpp
@@ -11,21 +11,25 @@
 IMPLEMENT_UNIFORM_BUFFER_STRUCT(FIrradianceCachingParameters, "IrradianceCachingParameters");
 IMPLEMENT_GLOBAL_SHADER(FVisualizeIrradianceCachePS, "/Plugin/GPULightmass/Private/IrradianceCacheVisualization.usf", "VisualizeIrradianceCachePS", SF_Pixel);
 
-FIrradianceCache::FIrradianceCache(int32 Quality, float Spacing, float CornerRejection)
+FIrradianceCache::FIrradianceCache(FRHICommandList& RHICmdList, int32 Quality, float Spacing, float CornerRejection)
 {
 	uint32 IrradianceCacheTotalBytes = 0;
 
-	FRHICommandListImmediate& RHICmdList = FRHICommandListExecutor::GetImmediateCommandList();
-
 	{
-		FRHIResourceCreateInfo CreateInfo(TEXT("FIrradianceCache"));
-		IrradianceCacheRecords = RHICmdList.CreateStructuredBuffer(sizeof(FVector4f), sizeof(FVector4f) * IrradianceCacheMaxSize, BUF_UnorderedAccess | BUF_ShaderResource, CreateInfo);
+		const FRHIBufferCreateDesc RecordsDesc =
+			FRHIBufferCreateDesc::CreateStructured<FVector4>(TEXT("FIrradianceCache"), IrradianceCacheMaxSize)
+			.AddUsage(BUF_UnorderedAccess | BUF_ShaderResource);
+		IrradianceCacheRecords = RHICmdList.CreateBuffer(RecordsDesc);
 		IrradianceCacheRecordsUAV = RHICmdList.CreateUnorderedAccessView(IrradianceCacheRecords, FRHIViewDesc::CreateBufferUAV().SetTypeFromBuffer(IrradianceCacheRecords));
-		IrradianceCacheRecordBackfaceHits = RHICmdList.CreateStructuredBuffer(sizeof(FVector4f), sizeof(FVector4f) * IrradianceCacheMaxSize, BUF_UnorderedAccess | BUF_ShaderResource, CreateInfo);
+
+		const FRHIBufferCreateDesc RecordBackfaceHitsDesc =
+			FRHIBufferCreateDesc::CreateStructured<FVector4>(TEXT("FIrradianceCache"), IrradianceCacheMaxSize)
+			.AddUsage(BUF_UnorderedAccess | BUF_ShaderResource);
+		IrradianceCacheRecordBackfaceHits = RHICmdList.CreateBuffer(RecordBackfaceHitsDesc);
 		IrradianceCacheRecordBackfaceHitsUAV = RHICmdList.CreateUnorderedAccessView(IrradianceCacheRecordBackfaceHits, FRHIViewDesc::CreateBufferUAV().SetTypeFromBuffer(IrradianceCacheRecordBackfaceHits));
 		
-		IrradianceCacheTotalBytes += sizeof(FVector4f) * IrradianceCacheMaxSize;
-		IrradianceCacheTotalBytes += sizeof(FVector4f) * IrradianceCacheMaxSize;
+		IrradianceCacheTotalBytes += RecordsDesc.Size;
+		IrradianceCacheTotalBytes += RecordBackfaceHitsDesc.Size;
 
 		RHICmdList.ClearUAVUint(IrradianceCacheRecordsUAV, FUintVector4(0, 0, 0, 0));
 		RHICmdList.ClearUAVUint(IrradianceCacheRecordBackfaceHitsUAV, FUintVector4(0, 0, 0, 0));
diff --git a/Engine/Plugins/Experimental/GPULightmass/Source/GPULightmass/Private/IrradianceCaching.h b/Engine/Plugins/Experimental/GPULightmass/Source/GPULightmass/Private/IrradianceCaching.h
--- a/Engine/Plugins/Experimental/GPULightmass/Source/GPULightmass/Private/IrradianceCaching.h
+++ b/Engine/Plugins/Experimental/GPULightmass/Source/GPULightmass/Private/IrradianceCaching.h
@@ -35,7 +35,7 @@ struct FIrradianceCache
 
 	TUniformBufferRef<FIrradianceCachingParameters> IrradianceCachingParametersUniformBuffer;
 
-	FIrradianceCache(int32 Quality, float Spacing, float CornerRejection);
+	FIrradianceCache(FRHICommandList& RHICmdList, int32 Quality, float Spacing, float CornerRejection);
 
 	FRWBuffer HashTable;
 	FRWBuffer HashToIndex;
diff --git a/Engine/Plugins/Experimental/GPULightmass/Source/GPULightmass/Private/LightmapRenderer.cpp b/Engine/Plugins/Experimental/GPULightmass/Source/GPULightmass/Private/LightmapRenderer.cpp
--- a/Engine/Plugins/Experimental/GPULightmass/Source/GPULightmass/Private/LightmapRenderer.cpp
+++ b/Engine/Plugins/Experimental/GPULightmass/Source/GPULightmass/Private/LightmapRenderer.cpp
@@ -208,17 +208,20 @@ struct FGPUBatchedTileRequests
 		{
 			FRHICommandListBase& RHICmdList = FRHICommandListImmediate::Get();
 
-			FRHIResourceCreateInfo CreateInfo(TEXT("BatchedTilesBuffer"));
-			CreateInfo.GPUMask = FRHIGPUMask::FromIndex(GPUIndex);
-			CreateInfo.ResourceArray = &BatchedTilesDesc;
-
-			BatchedTilesBuffer = RHICmdList.CreateStructuredBuffer(sizeof(FGPUTileDescription), BatchedTilesDesc.GetResourceDataSize(), BUF_Static | BUF_ShaderResource, CreateInfo);
+			const FRHIBufferCreateDesc CreateDesc =
+				FRHIBufferCreateDesc::CreateStructured<FGPUTileDescription>(TEXT("BatchedTilesBuffer"), BatchedTilesDesc.Num())
+				.AddUsage(EBufferUsageFlags::Static | EBufferUsageFlags::ShaderResource)
+				.SetGPUMask(FRHIGPUMask::FromIndex(GPUIndex))
+				.SetInitActionResourceArray(&BatchedTilesDesc)
+				.DetermineInitialState();
+
+			BatchedTilesBuffer = RHICmdList.CreateBuffer(CreateDesc);
 			BatchedTilesSRV = RHICmdList.CreateShaderResourceView(BatchedTilesBuffer, FRHIViewDesc::CreateBufferSRV().SetTypeFromBuffer(BatchedTilesBuffer));
 		}
 	}
 };
 
-FLightmapRenderer::FLightmapRenderer(FSceneRenderState* InScene)
+FLightmapRenderer::FLightmapRenderer(FRHICommandList& RHICmdList, FSceneRenderState* InScene)
 	: Scene(InScene)
 	, LightmapTilePoolGPU(FIntPoint(Scene->Settings->LightmapTilePoolSize))
 {
@@ -237,6 +240,7 @@ FLightmapRenderer::FLightmapRenderer(FSceneRenderState* InScene)
 	if (!Scene->Settings->bUseFirstBounceRayGuiding)
 	{
 		LightmapTilePoolGPU.Initialize(
+			RHICmdList,
 			{
 				{ PF_A32B32G32R32F, FIntPoint(GPreviewLightmapPhysicalTileSize) }, // IrradianceAndSampleCount
 				{ PF_A32B32G32R32F, FIntPoint(GPreviewLightmapPhysicalTileSize) }, // SHDirectionality
@@ -248,6 +252,7 @@ FLightmapRenderer::FLightmapRenderer(FSceneRenderState* InScene)
 	else
 	{
 		LightmapTilePoolGPU.Initialize(
+			RHICmdList,
 			{
 				{ PF_A32B32G32R32F, FIntPoint(GPreviewLightmapPhysicalTileSize) }, // IrradianceAndSampleCount
 				{ PF_A32B32G32R32F, FIntPoint(GPreviewLightmapPhysicalTileSize) }, // SHDirectionality
@@ -1085,8 +1090,12 @@ bool FSceneRenderState::SetupRayTracingScene(FRDGBuilder& GraphBuilder, FSceneUn
 							PRAGMA_DISABLE_DEPRECATION_WARNINGS
 							const uint32 BLASScratchSize = DynamicGeometryCollection.ComputeScratchBufferSize();
 
-							FRHIResourceCreateInfo ScratchBufferCreateInfo(TEXT("RHILandscapeScratchBuffer"));
-							FBufferRHIRef ScratchBuffer = RHICmdList.CreateBuffer(BLASScratchSize, BUF_StructuredBuffer | BUF_RayTracingScratch, 0, ERHIAccess::UAVCompute, ScratchBufferCreateInfo);
+							const FRHIBufferCreateDesc CreateDesc =
+								FRHIBufferCreateDesc::CreateStructured(TEXT("RHILandscapeScratchBuffer"), BLASScratchSize, 0)
+								.AddUsage(EBufferUsageFlags::RayTracingScratch)
+								.SetInitialState(ERHIAccess::UAVCompute);
+
+							FBufferRHIRef ScratchBuffer = RHICmdList.CreateBuffer(CreateDesc);
 
 							RHICmdList.SetStaticUniformBuffers({ReferenceView->ViewUniformBuffer, SceneUniforms.GetBufferRHI(GraphBuilder)});
 							DynamicGeometryCollection.DispatchUpdates(RHICmdList, ScratchBuffer);
@@ -1169,15 +1178,19 @@ bool FSceneRenderState::SetupRayTracingScene(FRDGBuilder& GraphBuilder, FSceneUn
 			const FRayTracingSceneInitializer& SceneInitializer = RayTracingScene->GetInitializer();
 
 			FRayTracingAccelerationStructureSize SizeInfo = RHICalcRayTracingSceneSize(SceneInitializer);
-			FRHIResourceCreateInfo BufferCreateInfo(TEXT("LightmassRayTracingSceneBuffer"));
-			RayTracingSceneBuffer = RHICmdList.CreateBuffer(uint32(SizeInfo.ResultSize), BUF_AccelerationStructure, 0, ERHIAccess::BVHWrite, BufferCreateInfo);
+
+			const FRHIBufferCreateDesc SceneCreateDesc =
+				FRHIBufferCreateDesc::Create(TEXT("LightmassRayTracingSceneBuffer"), SizeInfo.ResultSize, 0, EBufferUsageFlags::AccelerationStructure)
+				.SetInitialState(ERHIAccess::BVHWrite);
+
+			RayTracingSceneBuffer = RHICmdList.CreateBuffer(SceneCreateDesc);
 			RayTracingSceneSRV = RHICmdList.CreateShaderResourceView(FShaderResourceViewInitializer(RayTracingSceneBuffer, RayTracingScene, 0));
 
-			FRHIResourceCreateInfo ScratchBufferCreateInfo(TEXT("LightmassRayTracingScratchBuffer"));
-			FBufferRHIRef ScratchBuffer = RHICmdList.CreateBuffer(
-				uint32(SizeInfo.BuildScratchSize),
-				BUF_StructuredBuffer | BUF_RayTracingScratch, GRHIRayTracingScratchBufferAlignment, ERHIAccess::UAVCompute,
-				ScratchBufferCreateInfo);
+			const FRHIBufferCreateDesc ScratchCreateDesc =
+				FRHIBufferCreateDesc::Create(TEXT("LightmassRayTracingScratchBuffer"), SizeInfo.BuildScratchSize, GRHIRayTracingScratchBufferAlignment, EBufferUsageFlags::StructuredBuffer | EBufferUsageFlags::RayTracingScratch)
+				.SetInitialState(ERHIAccess::UAVCompute);
+
+			FBufferRHIRef ScratchBuffer = RHICmdList.CreateBuffer(ScratchCreateDesc);
 
 			FRWBufferStructured InstanceBuffer;
 			InstanceBuffer.Initialize(RHICmdList, TEXT("LightmassRayTracingInstanceBuffer"), GRHIRayTracingInstanceDescriptorSize, SceneInitializer.MaxNumInstances);
@@ -1191,38 +1204,29 @@ bool FSceneRenderState::SetupRayTracingScene(FRDGBuilder& GraphBuilder, FSceneUn
 				BUF_Volatile | BUF_MultiGPUAllocate);
 
 			const uint32 NumInstanceGroups = RayTracingSceneInitializationData.NumGPUInstanceGroups + RayTracingSceneInitializationData.NumCPUInstanceGroups;
-			const uint32 InstanceGroupUploadBytes = NumInstanceGroups * sizeof(FRayTracingInstanceGroup);
-			FBufferRHIRef InstanceGroupUploadBuffer;
-			FShaderResourceViewRHIRef InstanceGroupUploadSRV;
-			{
-				FRHIResourceCreateInfo CreateInfo(TEXT("RayTracingTestBedInstanceGroupUploadBuffer"));
-				InstanceGroupUploadBuffer = RHICmdList.CreateStructuredBuffer(sizeof(FRayTracingInstanceGroup), InstanceGroupUploadBytes, BUF_ShaderResource | BUF_Volatile, CreateInfo);
-				InstanceGroupUploadSRV = RHICmdList.CreateShaderResourceView(InstanceGroupUploadBuffer, FRHIViewDesc::CreateBufferSRV().SetTypeFromBuffer(InstanceGroupUploadBuffer));
-			}
+			const FRHIBufferCreateDesc InstanceGroupUploadCreateDesc =
+				FRHIBufferCreateDesc::CreateStructured<FRayTracingInstanceGroup>(TEXT("RayTracingTestBedInstanceGroupUploadBuffer"), NumInstanceGroups)
+				.AddUsage(EBufferUsageFlags::ShaderResource | EBufferUsageFlags::Volatile)
+				.SetInitActionInitializer()
+				.DetermineInitialState();
+			TRHIBufferInitializer<FRayTracingInstanceGroup> InstanceGroupUploadData = RHICmdList.CreateBufferInitializer(InstanceGroupUploadCreateDesc);
 
 			const uint32 NumInstanceDescriptors = RayTracingSceneInitializationData.NumGPUInstanceDescriptors + RayTracingSceneInitializationData.NumCPUInstanceDescriptors;
-			const uint32 InstanceUploadBufferSize = NumInstanceDescriptors * sizeof(FRayTracingInstanceDescriptor);
-			FBufferRHIRef InstanceUploadBuffer;
-			FShaderResourceViewRHIRef InstanceUploadSRV;
-			{
-				FRHIResourceCreateInfo CreateInfo(TEXT("LightmassRayTracingInstanceUploadBuffer"));
-				InstanceUploadBuffer = RHICmdList.CreateStructuredBuffer(sizeof(FRayTracingInstanceDescriptor), InstanceUploadBufferSize, BUF_ShaderResource | BUF_Volatile, CreateInfo);
-				InstanceUploadSRV = RHICmdList.CreateShaderResourceView(InstanceUploadBuffer, FRHIViewDesc::CreateBufferSRV().SetTypeFromBuffer(InstanceUploadBuffer));
-			}
+			const FRHIBufferCreateDesc InstanceUploadCreateDesc =
+				FRHIBufferCreateDesc::CreateStructured<FRayTracingInstanceDescriptor>(TEXT("LightmassRayTracingInstanceUploadBuffer"), NumInstanceDescriptors)
+				.AddUsage(EBufferUsageFlags::ShaderResource | EBufferUsageFlags::Volatile)
+				.SetInitActionInitializer()
+				.DetermineInitialState();
+			TRHIBufferInitializer<FRayTracingInstanceDescriptor> InstanceUploadData = RHICmdList.CreateBufferInitializer(InstanceUploadCreateDesc);
 
-			const uint32 TransformUploadBufferSize = RayTracingSceneInitializationData.NumNativeCPUInstances * 3 * sizeof(FVector4f);
-			FBufferRHIRef TransformUploadBuffer;
-			FShaderResourceViewRHIRef TransformUploadSRV;
-			{
-				FRHIResourceCreateInfo CreateInfo(TEXT("LightmassRayTracingTransformUploadBuffer"));
-				TransformUploadBuffer = RHICmdList.CreateStructuredBuffer(sizeof(FVector4f), TransformUploadBufferSize, BUF_ShaderResource | BUF_Volatile, CreateInfo);
-				TransformUploadSRV = RHICmdList.CreateShaderResourceView(TransformUploadBuffer, FRHIViewDesc::CreateBufferSRV().SetTypeFromBuffer(TransformUploadBuffer));
-			}
+			const FRHIBufferCreateDesc TransformUploadCreateDesc =
+				FRHIBufferCreateDesc::CreateStructured<FVector4f>(TEXT("LightmassRayTracingTransformUploadBuffer"), RayTracingSceneInitializationData.NumNativeCPUInstances * 3)
+				.AddUsage(EBufferUsageFlags::ShaderResource | EBufferUsageFlags::Volatile)
+				.SetInitActionInitializer()
+				.DetermineInitialState();
+			TRHIBufferInitializer<FVector4f> TransformUploadData = RHICmdList.CreateBufferInitializer(TransformUploadCreateDesc);
 
 			{
-				FRayTracingInstanceGroup* InstanceGroupUploadData = (FRayTracingInstanceGroup*)RHICmdList.LockBuffer(InstanceGroupUploadBuffer, 0, InstanceGroupUploadBytes, RLM_WriteOnly);
-				FRayTracingInstanceDescriptor* InstanceUploadData = (FRayTracingInstanceDescriptor*)RHICmdList.LockBuffer(InstanceUploadBuffer, 0, InstanceUploadBufferSize, RLM_WriteOnly);
-				FVector4f* TransformUploadData = (FVector4f*)RHICmdList.LockBuffer(TransformUploadBuffer, 0, TransformUploadBufferSize, RLM_WriteOnly);
 				FillRayTracingInstanceUploadBuffer(
 					View.ViewMatrices.GetPreViewTranslation(),
 					RayTracingGeometryInstances,
@@ -1234,14 +1238,21 @@ bool FSceneRenderState::SetupRayTracingScene(FRDGBuilder& GraphBuilder, FSceneUn
 					RayTracingSceneInitializationData.NumCPUInstanceGroups,
 					RayTracingSceneInitializationData.NumGPUInstanceDescriptors,
 					RayTracingSceneInitializationData.NumCPUInstanceDescriptors,
-					MakeArrayView(InstanceGroupUploadData, NumInstanceGroups),
-					MakeArrayView(InstanceUploadData, NumInstanceDescriptors),
-					MakeArrayView(TransformUploadData, RayTracingSceneInitializationData.NumNativeCPUInstances * 3));
-				RHICmdList.UnlockBuffer(TransformUploadBuffer);
-				RHICmdList.UnlockBuffer(InstanceUploadBuffer);
-				RHICmdList.UnlockBuffer(InstanceGroupUploadBuffer);
+					InstanceGroupUploadData.GetWriteView(),
+					InstanceUploadData.GetWriteView(),
+					TransformUploadData.GetWriteView()
+				);
 			}
 
+			FBufferRHIRef InstanceGroupUploadBuffer = InstanceGroupUploadData.Finalize();
+			FShaderResourceViewRHIRef InstanceGroupUploadSRV = RHICmdList.CreateShaderResourceView(InstanceGroupUploadBuffer, FRHIViewDesc::CreateBufferSRV().SetTypeFromBuffer(InstanceGroupUploadBuffer));
+
+			FBufferRHIRef InstanceUploadBuffer = InstanceUploadData.Finalize();
+			FShaderResourceViewRHIRef InstanceUploadSRV = RHICmdList.CreateShaderResourceView(InstanceUploadBuffer, FRHIViewDesc::CreateBufferSRV().SetTypeFromBuffer(InstanceUploadBuffer));
+
+			FBufferRHIRef TransformUploadBuffer = TransformUploadData.Finalize();
+			FShaderResourceViewRHIRef TransformUploadSRV = RHICmdList.CreateShaderResourceView(TransformUploadBuffer, FRHIViewDesc::CreateBufferSRV().SetTypeFromBuffer(TransformUploadBuffer));
+
 			for (uint32 GPUIndex : RHICmdList.GetGPUMask())
 			{
 				FRayTracingAccelerationStructureAddress* AddressesPtr = (FRayTracingAccelerationStructureAddress*)RHICmdList.LockBufferMGPU(
@@ -1875,16 +1886,26 @@ void FLightmapRenderer::Finalize(FRDGBuilder& GraphBuilder)
 					}
 
 					{
-						FRHIResourceCreateInfo CreateInfo(TEXT("SrcTilePositionsBuffer"), &SrcTilePositions);
-						CreateInfo.GPUMask = FRHIGPUMask::GPU0();
-						SrcTilePositionsBuffer = RHICmdList.CreateStructuredBuffer(sizeof(FIntPoint), SrcTilePositions.GetResourceDataSize(), BUF_Dynamic | BUF_ShaderResource, CreateInfo);
+						const FRHIBufferCreateDesc CreateDesc =
+							FRHIBufferCreateDesc::CreateStructured<FIntPoint>(TEXT("SrcTilePositionsBuffer"), SrcTilePositions.Num())
+							.AddUsage(EBufferUsageFlags::Dynamic | EBufferUsageFlags::ShaderResource)
+							.SetGPUMask(FRHIGPUMask::GPU0())
+							.SetInitActionResourceArray(&SrcTilePositions)
+							.DetermineInitialState();
+
+						SrcTilePositionsBuffer = RHICmdList.CreateBuffer(CreateDesc);
 						SrcTilePositionsSRV = RHICmdList.CreateShaderResourceView(SrcTilePositionsBuffer, FRHIViewDesc::CreateBufferSRV().SetTypeFromBuffer(SrcTilePositionsBuffer));
 					}
 
 					{
-						FRHIResourceCreateInfo CreateInfo(TEXT("DstTilePositionsBuffer"), &DstTilePositions);
-						CreateInfo.GPUMask = FRHIGPUMask::GPU0();
-						DstTilePositionsBuffer = RHICmdList.CreateStructuredBuffer(sizeof(FIntPoint), DstTilePositions.GetResourceDataSize(), BUF_Dynamic | BUF_ShaderResource, CreateInfo);
+						const FRHIBufferCreateDesc CreateDesc =
+							FRHIBufferCreateDesc::CreateStructured<FIntPoint>(TEXT("DstTilePositionsBuffer"), DstTilePositions.Num())
+							.AddUsage(EBufferUsageFlags::Dynamic | EBufferUsageFlags::ShaderResource)
+							.SetGPUMask(FRHIGPUMask::GPU0())
+							.SetInitActionResourceArray(&DstTilePositions)
+							.DetermineInitialState();
+
+						DstTilePositionsBuffer = RHICmdList.CreateBuffer(CreateDesc);
 						DstTilePositionsSRV = RHICmdList.CreateShaderResourceView(DstTilePositionsBuffer, FRHIViewDesc::CreateBufferSRV().SetTypeFromBuffer(DstTilePositionsBuffer));
 					}
 
@@ -2761,9 +2782,14 @@ void FLightmapRenderer::Finalize(FRDGBuilder& GraphBuilder)
 					check(PendingShadowTileRequests.Num() == LightTypeArray.Num());
 
 					{
-						FRHIResourceCreateInfo CreateInfo(TEXT("LightTypeBuffer"), &LightTypeArray);
-						CreateInfo.GPUMask = FRHIGPUMask::FromIndex(GPUIndex);
-						LightTypeBuffer = RHICmdList.CreateVertexBuffer(LightTypeArray.GetResourceDataSize(), BUF_Dynamic | BUF_ShaderResource, CreateInfo);
+						const FRHIBufferCreateDesc CreateDesc =
+							FRHIBufferCreateDesc::CreateVertex(TEXT("LightTypeBuffer"), LightTypeArray.GetResourceDataSize())
+							.AddUsage(EBufferUsageFlags::Dynamic | EBufferUsageFlags::ShaderResource)
+							.SetGPUMask(FRHIGPUMask::FromIndex(GPUIndex))
+							.SetInitActionResourceArray(&LightTypeArray)
+							.DetermineInitialState();
+
+						LightTypeBuffer = RHICmdList.CreateBuffer(CreateDesc);
 						LightTypeSRV = RHICmdList.CreateShaderResourceView(
 							LightTypeBuffer, 
 							FRHIViewDesc::CreateBufferSRV()
@@ -2772,9 +2798,14 @@ void FLightmapRenderer::Finalize(FRDGBuilder& GraphBuilder)
 					}
 
 					{
-						FRHIResourceCreateInfo CreateInfo(TEXT("ChannelIndexBuffer"), &ChannelIndexArray);
-						CreateInfo.GPUMask = FRHIGPUMask::FromIndex(GPUIndex);
-						ChannelIndexBuffer = RHICmdList.CreateVertexBuffer(ChannelIndexArray.GetResourceDataSize(), BUF_Dynamic | BUF_ShaderResource, CreateInfo);
+						const FRHIBufferCreateDesc CreateDesc =
+							FRHIBufferCreateDesc::CreateVertex(TEXT("ChannelIndexBuffer"), ChannelIndexArray.GetResourceDataSize())
+							.AddUsage(EBufferUsageFlags::Dynamic | EBufferUsageFlags::ShaderResource)
+							.SetGPUMask(FRHIGPUMask::FromIndex(GPUIndex))
+							.SetInitActionResourceArray(&ChannelIndexArray)
+							.DetermineInitialState();
+
+						ChannelIndexBuffer = RHICmdList.CreateBuffer(CreateDesc);
 						ChannelIndexSRV = RHICmdList.CreateShaderResourceView(
 							ChannelIndexBuffer, 
 							FRHIViewDesc::CreateBufferSRV()
@@ -2783,9 +2814,14 @@ void FLightmapRenderer::Finalize(FRDGBuilder& GraphBuilder)
 					}
 
 					{
-						FRHIResourceCreateInfo CreateInfo(TEXT("LightSampleIndexSRV"), &LightSampleIndexArray);
-						CreateInfo.GPUMask = FRHIGPUMask::FromIndex(GPUIndex);
-						LightSampleIndexBuffer = RHICmdList.CreateVertexBuffer(LightSampleIndexArray.GetResourceDataSize(), BUF_Dynamic | BUF_ShaderResource, CreateInfo);
+						const FRHIBufferCreateDesc CreateDesc =
+							FRHIBufferCreateDesc::CreateVertex(TEXT("LightSampleIndexSRV"), LightSampleIndexArray.GetResourceDataSize())
+							.AddUsage(EBufferUsageFlags::Dynamic | EBufferUsageFlags::ShaderResource)
+							.SetGPUMask(FRHIGPUMask::FromIndex(GPUIndex))
+							.SetInitActionResourceArray(&LightSampleIndexArray)
+							.DetermineInitialState();
+
+						LightSampleIndexBuffer = RHICmdList.CreateBuffer(CreateDesc);
 						LightSampleIndexSRV = RHICmdList.CreateShaderResourceView(
 							LightSampleIndexBuffer, 
 							FRHIViewDesc::CreateBufferSRV()
@@ -2794,9 +2830,14 @@ void FLightmapRenderer::Finalize(FRDGBuilder& GraphBuilder)
 					}
 
 					{
-						FRHIResourceCreateInfo CreateInfo(TEXT("LightShaderParameterBuffer"), &LightShaderParameterArray);
-						CreateInfo.GPUMask = FRHIGPUMask::FromIndex(GPUIndex);
-						LightShaderParameterBuffer = RHICmdList.CreateStructuredBuffer(sizeof(FLightShaderConstants), LightShaderParameterArray.GetResourceDataSize(), BUF_Dynamic | BUF_ShaderResource, CreateInfo);
+						const FRHIBufferCreateDesc CreateDesc =
+							FRHIBufferCreateDesc::CreateStructured<FLightShaderConstants>(TEXT("LightShaderParameterBuffer"), LightShaderParameterArray.Num())
+							.AddUsage(EBufferUsageFlags::Dynamic | EBufferUsageFlags::ShaderResource)
+							.SetGPUMask(FRHIGPUMask::FromIndex(GPUIndex))
+							.SetInitActionResourceArray(&LightShaderParameterArray)
+							.DetermineInitialState();
+
+						LightShaderParameterBuffer = RHICmdList.CreateBuffer(CreateDesc);
 						LightShaderParameterSRV = RHICmdList.CreateShaderResourceView(LightShaderParameterBuffer, FRHIViewDesc::CreateBufferSRV().SetTypeFromBuffer(LightShaderParameterBuffer));
 					}
 
diff --git a/Engine/Plugins/Experimental/GPULightmass/Source/GPULightmass/Private/LightmapRenderer.h b/Engine/Plugins/Experimental/GPULightmass/Source/GPULightmass/Private/LightmapRenderer.h
--- a/Engine/Plugins/Experimental/GPULightmass/Source/GPULightmass/Private/LightmapRenderer.h
+++ b/Engine/Plugins/Experimental/GPULightmass/Source/GPULightmass/Private/LightmapRenderer.h
@@ -110,7 +110,7 @@ public:
 class FLightmapRenderer : public IVirtualTextureFinalizer
 {
 public:
-	FLightmapRenderer(FSceneRenderState* InScene);
+	FLightmapRenderer(FRHICommandList& RHICmdList, FSceneRenderState* InScene);
 
 	void AddRequest(FLightmapTileRequest TileRequest);
 
diff --git a/Engine/Plugins/Experimental/GPULightmass/Source/GPULightmass/Private/LightmapTilePool.cpp b/Engine/Plugins/Experimental/GPULightmass/Source/GPULightmass/Private/LightmapTilePool.cpp
--- a/Engine/Plugins/Experimental/GPULightmass/Source/GPULightmass/Private/LightmapTilePool.cpp
+++ b/Engine/Plugins/Experimental/GPULightmass/Source/GPULightmass/Private/LightmapTilePool.cpp
@@ -149,7 +149,7 @@ FLightmapTilePoolGPU::FLightmapTilePoolGPU(int32 NumLayers, FIntPoint InSizeInTi
 	// UE_LOG(LogGPULightmass, Log, TEXT("LightmapTilePool created with %.2fMB"), MemorySize / 1024.0f / 1024.0f);
 }
 
-void FLightmapTilePoolGPU::Initialize(TArray<FLayerFormatAndTileSize> Layers)
+void FLightmapTilePoolGPU::Initialize(FRHICommandList& RHICmdList, TArray<FLayerFormatAndTileSize> Layers)
 {
 	ensure(SizeInTiles.X == SizeInTiles.Y);
 
@@ -173,8 +173,6 @@ void FLightmapTilePoolGPU::Initialize(TArray<FLayerFormatAndTileSize> Layers)
 			TexCreate_ShaderResource | TexCreate_RenderTargetable | TexCreate_UAV,
 			false);
 
-		FRHICommandListImmediate& RHICmdList = FRHICommandListExecutor::GetImmediateCommandList();
-
 		PooledRenderTargetDebugNames[RenderTargetIndex] = FString::Printf(TEXT("LightmapTilePoolGPU%d"), RenderTargetIndex);
 		GRenderTargetPool.FindFreeElement(RHICmdList, Desc, PooledRenderTargets[RenderTargetIndex], *PooledRenderTargetDebugNames[RenderTargetIndex]);
 
diff --git a/Engine/Plugins/Experimental/GPULightmass/Source/GPULightmass/Private/LightmapTilePool.h b/Engine/Plugins/Experimental/GPULightmass/Source/GPULightmass/Private/LightmapTilePool.h
--- a/Engine/Plugins/Experimental/GPULightmass/Source/GPULightmass/Private/LightmapTilePool.h
+++ b/Engine/Plugins/Experimental/GPULightmass/Source/GPULightmass/Private/LightmapTilePool.h
@@ -58,7 +58,7 @@ public:
 	FLightmapTilePoolGPU(FIntPoint InSizeInTiles) : FLightmapTilePool(InSizeInTiles.X * InSizeInTiles.Y), SizeInTiles(InSizeInTiles) {}
 	FLightmapTilePoolGPU(int32 NumLayers, FIntPoint InSizeInTiles, FIntPoint TileSize);
 
-	void Initialize(TArray<FLayerFormatAndTileSize> Layers);
+	void Initialize(FRHICommandList& RHICmdList, TArray<FLayerFormatAndTileSize> Layers);
 
 	FIntPoint SizeInTiles;
 
diff --git a/Engine/Plugins/Experimental/GPULightmass/Source/GPULightmass/Private/Scene/Scene.cpp b/Engine/Plugins/Experimental/GPULightmass/Source/GPULightmass/Private/Scene/Scene.cpp
--- a/Engine/Plugins/Experimental/GPULightmass/Source/GPULightmass/Private/Scene/Scene.cpp
+++ b/Engine/Plugins/Experimental/GPULightmass/Source/GPULightmass/Private/Scene/Scene.cpp
@@ -56,19 +56,19 @@ FScene::FScene(FGPULightmass* InGPULightmass)
 	RenderState.FeatureLevel = FeatureLevel;
 
 	ENQUEUE_RENDER_COMMAND(RenderThreadInit)(
-		[&RenderState = RenderState](FRHICommandListImmediate&) mutable
+		[&RenderState = RenderState](FRHICommandListImmediate& RHICmdList) mutable
 	{
-		RenderState.RenderThreadInit();
+		RenderState.RenderThreadInit(RHICmdList);
 	});
 }
 
-void FSceneRenderState::RenderThreadInit()
+void FSceneRenderState::RenderThreadInit(FRHICommandList& RHICmdList)
 {
 	check(IsInRenderingThread());
 
-	LightmapRenderer = MakeUnique<FLightmapRenderer>(this);
+	LightmapRenderer = MakeUnique<FLightmapRenderer>(RHICmdList, this);
 	VolumetricLightmapRenderer = MakeUnique<FVolumetricLightmapRenderer>(this);
-	IrradianceCache = MakeUnique<FIrradianceCache>(Settings->IrradianceCacheQuality, Settings->IrradianceCacheSpacing, Settings->IrradianceCacheCornerRejection);
+	IrradianceCache = MakeUnique<FIrradianceCache>(RHICmdList, Settings->IrradianceCacheQuality, Settings->IrradianceCacheSpacing, Settings->IrradianceCacheCornerRejection);
 	IrradianceCache->CurrentRevision = LightmapRenderer->GetCurrentRevision();
 }
 
@@ -1568,8 +1568,10 @@ void FScene::BackgroundTick()
 				bNeedsVoxelization = false;
 			}
 
-			ENQUEUE_RENDER_COMMAND(BackgroundTickRenderThread)([&RenderState = RenderState](FRHICommandListImmediate&) mutable {
-				RenderState.BackgroundTick();
+			ENQUEUE_RENDER_COMMAND(BackgroundTickRenderThread)(
+			[&RenderState = RenderState](FRHICommandListImmediate& RHICmdList) mutable
+			{
+				RenderState.BackgroundTick(RHICmdList);
 			});
 		}
 	}
@@ -1585,11 +1587,11 @@ void FScene::BackgroundTick()
 	}
 }
 
-void FSceneRenderState::BackgroundTick()
+void FSceneRenderState::BackgroundTick(FRHICommandList& RHICmdList)
 {
 	if (IrradianceCache->CurrentRevision != LightmapRenderer->GetCurrentRevision())
 	{
-		IrradianceCache = MakeUnique<FIrradianceCache>(Settings->IrradianceCacheQuality, Settings->IrradianceCacheSpacing, Settings->IrradianceCacheCornerRejection);
+		IrradianceCache = MakeUnique<FIrradianceCache>(RHICmdList, Settings->IrradianceCacheQuality, Settings->IrradianceCacheSpacing, Settings->IrradianceCacheCornerRejection);
 		IrradianceCache->CurrentRevision = LightmapRenderer->GetCurrentRevision();
 	}
 
diff --git a/Engine/Plugins/Experimental/GPULightmass/Source/GPULightmass/Private/Scene/Scene.h b/Engine/Plugins/Experimental/GPULightmass/Source/GPULightmass/Private/Scene/Scene.h
--- a/Engine/Plugins/Experimental/GPULightmass/Source/GPULightmass/Private/Scene/Scene.h
+++ b/Engine/Plugins/Experimental/GPULightmass/Source/GPULightmass/Private/Scene/Scene.h
@@ -100,8 +100,8 @@ class FSceneRenderState
 public:
 	UGPULightmassSettings* Settings;
 
-	void RenderThreadInit();
-	void BackgroundTick();
+	void RenderThreadInit(FRHICommandList& RHICmdList);
+	void BackgroundTick(FRHICommandList& RHICmdList);
 
 	FRayTracingSceneRHIRef RayTracingScene;
 	FShaderResourceViewRHIRef RayTracingSceneSRV;
diff --git a/Engine/Plugins/Experimental/ImagePlate/Source/ImagePlate/Private/ImagePlateComponent.cpp b/Engine/Plugins/Experimental/ImagePlate/Source/ImagePlate/Private/ImagePlateComponent.cpp
--- a/Engine/Plugins/Experimental/ImagePlate/Source/ImagePlate/Private/ImagePlateComponent.cpp
+++ b/Engine/Plugins/Experimental/ImagePlate/Source/ImagePlate/Private/ImagePlateComponent.cpp
@@ -20,6 +20,7 @@
 #include "SceneInterface.h"
 #include "StaticMeshResources.h"
 #include "Kismet/GameplayStatics.h"
+#include "RHIResourceUtils.h"
 
 #include UE_INLINE_GENERATED_CPP_BY_NAME(ImagePlateComponent)
 
@@ -39,15 +40,7 @@ namespace
 
 		virtual void InitRHI(FRHICommandListBase& RHICmdList) override
 		{
-			const uint32 Size = Indices.Num() * sizeof(uint16);
-
-			FRHIResourceCreateInfo CreateInfo(TEXT("FImagePlateIndexBuffer"));
-			IndexBufferRHI = RHICmdList.CreateBuffer(Size, BUF_Static | BUF_IndexBuffer, sizeof(uint16), ERHIAccess::VertexOrIndexBuffer, CreateInfo);
-
-			// Copy the index data into the index buffer.		
-			void* Buffer = RHICmdList.LockBuffer(IndexBufferRHI, 0, Size, RLM_WriteOnly);
-			FMemory::Memcpy(Buffer, Indices.GetData(), Size);
-			RHICmdList.UnlockBuffer(IndexBufferRHI);
+			IndexBufferRHI = UE::RHIResourceUtils::CreateIndexBufferFromArray(RHICmdList, TEXT("FImagePlateIndexBuffer"), EBufferUsageFlags::Static, MakeConstArrayView(Indices));
 		}
 	};
 
diff --git a/Engine/Plugins/Experimental/SimpleHMD/Source/SimpleHMD/Private/SimpleHMD.cpp b/Engine/Plugins/Experimental/SimpleHMD/Source/SimpleHMD/Private/SimpleHMD.cpp
--- a/Engine/Plugins/Experimental/SimpleHMD/Source/SimpleHMD/Private/SimpleHMD.cpp
+++ b/Engine/Plugins/Experimental/SimpleHMD/Source/SimpleHMD/Private/SimpleHMD.cpp
@@ -7,6 +7,7 @@
 #include "PostProcess/PostProcessHMD.h"
 #include "GameFramework/WorldSettings.h"
 #include "SceneView.h"
+#include "RHIResourceUtils.h"
 
 //---------------------------------------------------
 // SimpleHMD Plugin Implementation
@@ -198,19 +199,10 @@ void FSimpleHMD::DrawDistortionMesh_RenderThread(struct FHeadMountedDisplayPassC
 			{ FVector2f(0.1f, 0.9f), FVector2f(0.5f, 0.0f), FVector2f(0.5f, 0.0f), FVector2f(0.5f, 0.0f), 1.0f, 0.0f },
 		}
 	};
-
-	FRHIResourceCreateInfo CreateInfo(TEXT("FSimpleHMD"));
-	FBufferRHIRef VertexBufferRHI = RHICmdList.CreateVertexBuffer(sizeof(FDistortionVertex) * NumVerts, BUF_Volatile, CreateInfo);
-	void* VoidPtr = RHICmdList.LockBuffer(VertexBufferRHI, 0, sizeof(FDistortionVertex) * NumVerts, RLM_WriteOnly);
-	FPlatformMemory::Memcpy(VoidPtr, MeshVerts[View.StereoViewIndex], sizeof(FDistortionVertex) * NumVerts);
-	RHICmdList.UnlockBuffer(VertexBufferRHI);
-
 	static const uint16 Indices[] = { 0, 1, 2, 0, 2, 3 };
 
-	FBufferRHIRef IndexBufferRHI = RHICmdList.CreateIndexBuffer(sizeof(uint16), sizeof(uint16) * 6, BUF_Volatile, CreateInfo);
-	void* VoidPtr2 = RHICmdList.LockBuffer(IndexBufferRHI, 0, sizeof(uint16) * 6, RLM_WriteOnly);
-	FPlatformMemory::Memcpy(VoidPtr2, Indices, sizeof(uint16) * 6);
-	RHICmdList.UnlockBuffer(IndexBufferRHI);
+	FBufferRHIRef VertexBufferRHI = UE::RHIResourceUtils::CreateVertexBufferFromArray(RHICmdList, TEXT("FSimpleHMD"), EBufferUsageFlags::Volatile, MakeConstArrayView(MeshVerts[View.StereoViewIndex]));
+	FBufferRHIRef IndexBufferRHI = UE::RHIResourceUtils::CreateIndexBufferFromArray(RHICmdList, TEXT("FSimpleHMD"), EBufferUsageFlags::Volatile, MakeConstArrayView(Indices));
 
 	RHICmdList.SetStreamSource(0, VertexBufferRHI, 0);
 	RHICmdList.DrawIndexedPrimitive(IndexBufferRHI, 0, 0, NumVerts, 0, NumTris, 1);
diff --git a/Engine/Plugins/Experimental/VirtualHeightfieldMesh/Source/VirtualHeightfieldMesh/Private/VirtualHeightfieldMeshSceneProxy.cpp b/Engine/Plugins/Experimental/VirtualHeightfieldMesh/Source/VirtualHeightfieldMesh/Private/VirtualHeightfieldMeshSceneProxy.cpp
--- a/Engine/Plugins/Experimental/VirtualHeightfieldMesh/Source/VirtualHeightfieldMesh/Private/VirtualHeightfieldMeshSceneProxy.cpp
+++ b/Engine/Plugins/Experimental/VirtualHeightfieldMesh/Source/VirtualHeightfieldMesh/Private/VirtualHeightfieldMeshSceneProxy.cpp
@@ -993,10 +993,12 @@ namespace VirtualHeightfieldMesh
 	void InitializeInstanceBuffers(FRHICommandListBase& RHICmdList, FDrawInstanceBuffers& InBuffers)
 	{
 		{
-			FRHIResourceCreateInfo CreateInfo(TEXT("VirtualHeightfieldMesh.InstanceBuffer"));
-			const int32 InstanceSize = sizeof(VirtualHeightfieldMesh::QuadRenderInstance);
-			const int32 InstanceBufferSize = CVarVHMMaxRenderItems.GetValueOnRenderThread() * InstanceSize;
-			InBuffers.InstanceBuffer = RHICmdList.CreateStructuredBuffer(InstanceSize, InstanceBufferSize, BUF_UnorderedAccess|BUF_ShaderResource, ERHIAccess::SRVMask, CreateInfo);
+			const FRHIBufferCreateDesc CreateDesc =
+				FRHIBufferCreateDesc::CreateStructured<VirtualHeightfieldMesh::QuadRenderInstance>(TEXT("VirtualHeightfieldMesh.InstanceBuffer"), CVarVHMMaxRenderItems.GetValueOnRenderThread())
+				.AddUsage(EBufferUsageFlags::UnorderedAccess | EBufferUsageFlags::ShaderResource)
+				.SetInitialState(ERHIAccess::SRVMask);
+
+			InBuffers.InstanceBuffer = RHICmdList.CreateBuffer(CreateDesc);
 			InBuffers.InstanceBufferUAV = RHICmdList.CreateUnorderedAccessView(
 				InBuffers.InstanceBuffer, 
 				FRHIViewDesc::CreateBufferUAV()
@@ -1004,8 +1006,12 @@ namespace VirtualHeightfieldMesh
 			InBuffers.InstanceBufferSRV = RHICmdList.CreateShaderResourceView(InBuffers.InstanceBuffer, FRHIViewDesc::CreateBufferSRV().SetTypeFromBuffer(InBuffers.InstanceBuffer));
 		}
 		{
-			FRHIResourceCreateInfo CreateInfo(TEXT("VirtualHeightfieldMesh.InstanceIndirectArgsBuffer"));
-			InBuffers.IndirectArgsBuffer = RHICmdList.CreateVertexBuffer(5 * sizeof(uint32), BUF_UnorderedAccess|BUF_DrawIndirect, ERHIAccess::IndirectArgs, CreateInfo);
+			const FRHIBufferCreateDesc CreateDesc =
+				FRHIBufferCreateDesc::CreateVertex<uint32>(TEXT("VirtualHeightfieldMesh.InstanceIndirectArgsBuffer"), 5)
+				.AddUsage(EBufferUsageFlags::UnorderedAccess | EBufferUsageFlags::DrawIndirect)
+				.SetInitialState(ERHIAccess::IndirectArgs);
+
+			InBuffers.IndirectArgsBuffer = RHICmdList.CreateBuffer(CreateDesc);
 			InBuffers.IndirectArgsBufferUAV = RHICmdList.CreateUnorderedAccessView(
 				InBuffers.IndirectArgsBuffer, 
 				FRHIViewDesc::CreateBufferUAV()
diff --git a/Engine/Plugins/Experimental/Water/Source/Runtime/Private/WaterQuadTreeGPU.cpp b/Engine/Plugins/Experimental/Water/Source/Runtime/Private/WaterQuadTreeGPU.cpp
--- a/Engine/Plugins/Experimental/Water/Source/Runtime/Private/WaterQuadTreeGPU.cpp
+++ b/Engine/Plugins/Experimental/Water/Source/Runtime/Private/WaterQuadTreeGPU.cpp
@@ -426,13 +426,17 @@ public:
 	virtual void InitRHI(FRHICommandListBase& RHICmdList) override
 	{
 		// Create the texture RHI.
-		FRHIResourceCreateInfo CreateInfo(TEXT("JitterOffsetVertexBuffer"));
-
 		const int32 NumMSAASamples = 2 + 4 + 8 + 16;
 		const int32 NumFloats = 2 * (NumHaltonSamples + NumMSAASamples);
-		VertexBufferRHI = RHICmdList.CreateVertexBuffer(sizeof(float) * NumFloats, BUF_Static | BUF_ShaderResource, CreateInfo);
 
-		float* BufferData = (float*)RHICmdList.LockBuffer(VertexBufferRHI, 0, sizeof(float) * NumFloats, RLM_WriteOnly);
+		const FRHIBufferCreateDesc CreateDesc =
+			FRHIBufferCreateDesc::CreateVertex<float>(TEXT("JitterOffsetVertexBuffer"), NumFloats)
+			.AddUsage(EBufferUsageFlags::Static | EBufferUsageFlags::ShaderResource)
+			.SetInitActionInitializer()
+			.DetermineInitialState();
+
+		TRHIBufferInitializer<float> Initializer = RHICmdList.CreateBufferInitializer(CreateDesc);
+		float* BufferData = Initializer.GetWritableData();
 		
 		// Halton
 		for (int32 i = 0; i < NumHaltonSamples; ++i)
@@ -460,7 +464,7 @@ public:
 			}
 		}
 
-		RHICmdList.UnlockBuffer(VertexBufferRHI);
+		VertexBufferRHI = Initializer.Finalize();
 	}
 
 	static constexpr int32 GetMSAADataOffset(int32 NumSamples)
diff --git a/Engine/Plugins/Experimental/Water/Source/Runtime/Public/WaterInstanceDataBuffer.h b/Engine/Plugins/Experimental/Water/Source/Runtime/Public/WaterInstanceDataBuffer.h
--- a/Engine/Plugins/Experimental/Water/Source/Runtime/Public/WaterInstanceDataBuffer.h
+++ b/Engine/Plugins/Experimental/Water/Source/Runtime/Public/WaterInstanceDataBuffer.h
@@ -16,13 +16,16 @@ public:
 		(
 			[this, InInstanceCount](FRHICommandListImmediate& RHICmdList)
 			{
-				FRHIResourceCreateInfo CreateInfo(TEXT("WaterInstanceDataBuffers"));
+				const int32 SizeInBytes = Align<int32>(InInstanceCount * sizeof(FVector4f), 4 * 1024);
 
-				int32 SizeInBytes = Align<int32>(InInstanceCount * sizeof(FVector4f), 4 * 1024);
+				const FRHIBufferCreateDesc CreateDesc =
+					FRHIBufferCreateDesc::CreateVertex(TEXT("WaterInstanceDataBuffers"), SizeInBytes)
+					.AddUsage(EBufferUsageFlags::Dynamic)
+					.DetermineInitialState();
 
 				for (int32 i = 0; i < NumBuffers; ++i)
 				{
-					Buffer[i] = RHICmdList.CreateVertexBuffer(SizeInBytes, BUF_Dynamic, CreateInfo);
+					Buffer[i] = RHICmdList.CreateBuffer(CreateDesc);
 					BufferMemory[i] = TArrayView<FVector4f>();
 				}
 			}
@@ -74,12 +77,15 @@ private:
 		{
 			Buffer[InBufferID].SafeRelease();
 
-			FRHIResourceCreateInfo CreateInfo(TEXT("WaterInstanceDataBuffers"));
-
 			// Align size in to avoid reallocating for a few differences of instance count
 			uint32 AlignedSizeInBytes = Align<uint32>(SizeInBytes, 4 * 1024);
 
-			Buffer[InBufferID] = RHICmdList.CreateVertexBuffer(AlignedSizeInBytes, BUF_Dynamic, CreateInfo);
+			const FRHIBufferCreateDesc CreateDesc =
+				FRHIBufferCreateDesc::CreateVertex(TEXT("WaterInstanceDataBuffers"), AlignedSizeInBytes)
+				.AddUsage(EBufferUsageFlags::Dynamic)
+				.DetermineInitialState();
+
+			Buffer[InBufferID] = RHICmdList.CreateBuffer(CreateDesc);
 		}
 
 		FVector4f* Data = reinterpret_cast<FVector4f*>(RHICmdList.LockBuffer(Buffer[InBufferID], 0, SizeInBytes, RLM_WriteOnly));
diff --git a/Engine/Plugins/Experimental/Water/Source/Runtime/Public/WaterVertexFactory.h b/Engine/Plugins/Experimental/Water/Source/Runtime/Public/WaterVertexFactory.h
--- a/Engine/Plugins/Experimental/Water/Source/Runtime/Public/WaterVertexFactory.h
+++ b/Engine/Plugins/Experimental/Water/Source/Runtime/Public/WaterVertexFactory.h
@@ -118,9 +118,13 @@ public:
 
 		NumVerts = NumVertsPerSide * NumVertsPerSide;
 
-		FRHIResourceCreateInfo CreateInfo(TEXT("FWaterMeshVertexBuffer"));
-		VertexBufferRHI = RHICmdList.CreateBuffer(sizeof(FVector4f) * NumVerts, BUF_Static | BUF_VertexBuffer | BUF_ShaderResource, 0, ERHIAccess::VertexOrIndexBuffer | ERHIAccess::SRVMask, CreateInfo);
-		FVector4f* DummyContents = (FVector4f*)RHICmdList.LockBuffer(VertexBufferRHI, 0, sizeof(FVector4f) * NumVerts, RLM_WriteOnly);
+		const FRHIBufferCreateDesc CreateDesc =
+			FRHIBufferCreateDesc::CreateVertex<FVector4f>(TEXT("FWaterMeshVertexBuffer"), NumVerts)
+			.AddUsage(EBufferUsageFlags::Static | EBufferUsageFlags::ShaderResource)
+			.SetInitialState(ERHIAccess::VertexOrIndexBuffer | ERHIAccess::SRVMask)
+			.SetInitActionInitializer();
+
+		TRHIBufferInitializer<FVector4f> DummyContents = RHICmdList.CreateBufferInitializer(CreateDesc);
 
 		for (uint32 VertY = 0; VertY < NumVertsPerSide; VertY++)
 		{
@@ -135,8 +139,7 @@ public:
 			}
 		}
 
-		RHICmdList.UnlockBuffer(VertexBufferRHI);
-
+		VertexBufferRHI = DummyContents.Finalize();
 		SRV = RHICmdList.CreateShaderResourceView(
 			VertexBufferRHI, 
 			FRHIViewDesc::CreateBufferSRV()
diff --git a/Engine/Plugins/FX/Niagara/Source/Niagara/Private/Experimental/NiagaraMeshUvMapping.cpp b/Engine/Plugins/FX/Niagara/Source/Niagara/Private/Experimental/NiagaraMeshUvMapping.cpp
--- a/Engine/Plugins/FX/Niagara/Source/Niagara/Private/Experimental/NiagaraMeshUvMapping.cpp
+++ b/Engine/Plugins/FX/Niagara/Source/Niagara/Private/Experimental/NiagaraMeshUvMapping.cpp
@@ -118,12 +118,15 @@ FMeshUvMappingBufferProxy::Initialize(const FMeshUvMapping& UvMapping)
 void
 FMeshUvMappingBufferProxy::InitRHI(FRHICommandListBase& RHICmdList)
 {
-	FRHIResourceCreateInfo CreateInfo(TEXT("UvMappingBuffer"));
-	CreateInfo.ResourceArray = &FrozenQuadTree;
+	const int32 BufferSize = FrozenQuadTree.GetResourceDataSize();
 
-	const int32 BufferSize = FrozenQuadTree.Num();
+	const FRHIBufferCreateDesc CreateDesc =
+		FRHIBufferCreateDesc::CreateVertex(TEXT("UvMappingBuffer"), BufferSize)
+		.AddUsage(EBufferUsageFlags::ShaderResource | EBufferUsageFlags::Static)
+		.SetInitActionResourceArray(&FrozenQuadTree)
+		.DetermineInitialState();
 
-	UvMappingBuffer = RHICmdList.CreateVertexBuffer(BufferSize, BUF_ShaderResource | BUF_Static, CreateInfo);
+	UvMappingBuffer = RHICmdList.CreateBuffer(CreateDesc);
 	UvMappingSrv = RHICmdList.CreateShaderResourceView(
 		UvMappingBuffer, 
 		FRHIViewDesc::CreateBufferSRV()
diff --git a/Engine/Plugins/FX/Niagara/Source/Niagara/Private/NiagaraDataInterfaceArrayImpl.cpp b/Engine/Plugins/FX/Niagara/Source/Niagara/Private/NiagaraDataInterfaceArrayImpl.cpp
--- a/Engine/Plugins/FX/Niagara/Source/Niagara/Private/NiagaraDataInterfaceArrayImpl.cpp
+++ b/Engine/Plugins/FX/Niagara/Source/Niagara/Private/NiagaraDataInterfaceArrayImpl.cpp
@@ -263,10 +263,14 @@ void FNDIArrayInstanceData_RenderThreadBase::UpdateDataInternal(FRHICommandList&
 		const int32 TypeStride = GPixelFormats[PixelFormat].BlockBytes;
 
 		// Create Buffer
-		FRHIResourceCreateInfo CreateInfo(TEXT("NiagaraDataInterfaceArray"));
 		const EBufferUsageFlags BufferUsage = BUF_Static | BUF_ShaderResource | BUF_VertexBuffer | BUF_SourceCopy | (IsReadOnly() ? BUF_None : BUF_UnorderedAccess);
 		const ERHIAccess DefaultAccess = IsReadOnly() ? ERHIAccess::SRVCompute : ERHIAccess::UAVCompute;
-		ArrayBuffer = RHICmdList.CreateBuffer(ArrayNumBytes, BufferUsage, TypeStride, DefaultAccess, CreateInfo);
+
+		const FRHIBufferCreateDesc CreateDesc =
+			FRHIBufferCreateDesc::Create(TEXT("NiagaraDataInterfaceArray"), ArrayNumBytes, TypeStride, BufferUsage)
+			.SetInitialState(DefaultAccess);
+
+		ArrayBuffer = RHICmdList.CreateBuffer(CreateDesc);
 
 		ArraySRV = RHICmdList.CreateShaderResourceView(
 			ArrayBuffer, 
diff --git a/Engine/Plugins/FX/Niagara/Source/Niagara/Private/NiagaraDataInterfaceSkeletalMesh.cpp b/Engine/Plugins/FX/Niagara/Source/Niagara/Private/NiagaraDataInterfaceSkeletalMesh.cpp
--- a/Engine/Plugins/FX/Niagara/Source/Niagara/Private/NiagaraDataInterfaceSkeletalMesh.cpp
+++ b/Engine/Plugins/FX/Niagara/Source/Niagara/Private/NiagaraDataInterfaceSkeletalMesh.cpp
@@ -964,12 +964,19 @@ void FSkeletalMeshGpuSpawnStaticBuffers::InitRHI(FRHICommandListBase& RHICmdList
 		const FSkeletalMeshAreaWeightedTriangleSampler& triangleSampler = SkeletalMeshSamplingLODBuiltData->AreaWeightedTriangleSampler;
 		check(TriangleCount == triangleSampler.GetNumEntries());
 
-		FRHIResourceCreateInfo CreateInfo(TEXT("FSkeletalMeshGpuSpawnStaticBuffers"));
-		uint32 SizeByte = NDISkelMeshLocal::GetProbAliasDWORDSize(TriangleCount) * sizeof(uint32);
-		BufferTriangleUniformSamplerProbAliasRHI = RHICmdList.CreateBuffer(SizeByte, BUF_Static | BUF_VertexBuffer | BUF_ShaderResource, 0, ERHIAccess::VertexOrIndexBuffer | ERHIAccess::SRVMask, CreateInfo);
-		uint32* PackedData = (uint32*)RHICmdList.LockBuffer(BufferTriangleUniformSamplerProbAliasRHI, 0, SizeByte, RLM_WriteOnly);
-		NDISkelMeshLocal::PackProbAlias(PackedData, triangleSampler);
-		RHICmdList.UnlockBuffer(BufferTriangleUniformSamplerProbAliasRHI);
+		const uint32 SizeByte = NDISkelMeshLocal::GetProbAliasDWORDSize(TriangleCount) * sizeof(uint32);
+
+		const FRHIBufferCreateDesc CreateDesc =
+			FRHIBufferCreateDesc::CreateVertex(TEXT("FSkeletalMeshGpuSpawnStaticBuffers"), SizeByte)
+			.AddUsage(EBufferUsageFlags::Static | EBufferUsageFlags::ShaderResource)
+			.SetInitialState(ERHIAccess::VertexOrIndexBuffer | ERHIAccess::SRVMask)
+			.SetInitActionInitializer();
+
+		TRHIBufferInitializer<uint32> Initializer = RHICmdList.CreateBufferInitializer(CreateDesc);
+
+		NDISkelMeshLocal::PackProbAlias(Initializer.GetWritableData(), triangleSampler);
+
+		BufferTriangleUniformSamplerProbAliasRHI = Initializer.Finalize();
 		BufferTriangleUniformSamplerProbAliasSRV = RHICmdList.CreateShaderResourceView(
 			BufferTriangleUniformSamplerProbAliasRHI, 
 			FRHIViewDesc::CreateBufferSRV()
@@ -990,11 +997,14 @@ void FSkeletalMeshGpuSpawnStaticBuffers::InitRHI(FRHICommandListBase& RHICmdList
 	SampleRegionsVerticesSRV = FNiagaraRenderer::GetDummyUIntBuffer();
 	if (NumSamplingRegionTriangles > 0)
 	{
-		FRHIResourceCreateInfo CreateInfo(TEXT("SampleRegionsProbAliasBuffer"));
 		if (bSamplingRegionsAllAreaWeighted)
 		{
-			CreateInfo.ResourceArray = &SampleRegionsProbAlias;
-			SampleRegionsProbAliasBuffer = RHICmdList.CreateVertexBuffer(SampleRegionsProbAlias.Num() * SampleRegionsProbAlias.GetTypeSize(), BUF_Static | BUF_ShaderResource, CreateInfo);
+			const FRHIBufferCreateDesc CreateDesc =
+				FRHIBufferCreateDesc::CreateVertex(TEXT("SampleRegionsProbAliasBuffer"), SampleRegionsProbAlias.Num() * SampleRegionsProbAlias.GetTypeSize())
+				.AddUsage(EBufferUsageFlags::Static | EBufferUsageFlags::ShaderResource)
+				.SetInitActionResourceArray(&SampleRegionsProbAlias);
+
+			SampleRegionsProbAliasBuffer = RHICmdList.CreateBuffer(CreateDesc);
 			SampleRegionsProbAliasSRV = RHICmdList.CreateShaderResourceView(
 				SampleRegionsProbAliasBuffer, 
 				FRHIViewDesc::CreateBufferSRV()
@@ -1009,29 +1019,39 @@ void FSkeletalMeshGpuSpawnStaticBuffers::InitRHI(FRHICommandListBase& RHICmdList
 			SampleRegionsProbAliasSRV = FNiagaraRenderer::GetDummyUIntBuffer();
 		}
 
-		CreateInfo.DebugName = (TEXT("SampleRegionsTriangleIndicesBuffer"));
-		CreateInfo.ResourceArray = &SampleRegionsTriangleIndicies;
-		SampleRegionsTriangleIndicesBuffer = RHICmdList.CreateVertexBuffer(SampleRegionsTriangleIndicies.Num() * SampleRegionsTriangleIndicies.GetTypeSize(), BUF_Static | BUF_ShaderResource, CreateInfo);
+		{
+			const FRHIBufferCreateDesc CreateDesc =
+				FRHIBufferCreateDesc::CreateVertex(TEXT("SampleRegionsTriangleIndicesBuffer"), SampleRegionsTriangleIndicies.Num() * SampleRegionsTriangleIndicies.GetTypeSize())
+				.AddUsage(EBufferUsageFlags::Static | EBufferUsageFlags::ShaderResource)
+				.SetInitActionResourceArray(&SampleRegionsTriangleIndicies);
+
+			SampleRegionsTriangleIndicesBuffer = RHICmdList.CreateBuffer(CreateDesc);
 		SampleRegionsTriangleIndicesSRV = RHICmdList.CreateShaderResourceView(
 			SampleRegionsTriangleIndicesBuffer, 
 			FRHIViewDesc::CreateBufferSRV()
 				.SetType(FRHIViewDesc::EBufferType::Typed)
 				.SetFormat(PF_R32_UINT));
 #if STATS
-		GPUMemoryUsage += SampleRegionsTriangleIndicies.Num() * SampleRegionsTriangleIndicies.GetTypeSize();
+			GPUMemoryUsage += SampleRegionsTriangleIndicies.Num() * SampleRegionsTriangleIndicies.GetTypeSize();
 #endif
+		}
 
-		CreateInfo.DebugName = (TEXT("SampleRegionsVerticesBuffer"));
-		CreateInfo.ResourceArray = &SampleRegionsVertices;
-		SampleRegionsVerticesBuffer = RHICmdList.CreateVertexBuffer(SampleRegionsVertices.Num() * SampleRegionsVertices.GetTypeSize(), BUF_Static | BUF_ShaderResource, CreateInfo);
+		{
+			const FRHIBufferCreateDesc CreateDesc =
+				FRHIBufferCreateDesc::CreateVertex(TEXT("SampleRegionsVerticesBuffer"), SampleRegionsVertices.Num() * SampleRegionsVertices.GetTypeSize())
+				.AddUsage(EBufferUsageFlags::Static | EBufferUsageFlags::ShaderResource)
+				.SetInitActionResourceArray(&SampleRegionsVertices);
+
+			SampleRegionsVerticesBuffer = RHICmdList.CreateBuffer(CreateDesc);
 		SampleRegionsVerticesSRV = RHICmdList.CreateShaderResourceView(
 			SampleRegionsVerticesBuffer, 
 			FRHIViewDesc::CreateBufferSRV()
 				.SetType(FRHIViewDesc::EBufferType::Typed)
 				.SetFormat(PF_R32_UINT));
 #if STATS
-		GPUMemoryUsage += SampleRegionsVertices.Num() * SampleRegionsVertices.GetTypeSize();
+			GPUMemoryUsage += SampleRegionsVertices.Num() * SampleRegionsVertices.GetTypeSize();
 #endif
+		}
 	}
 	else
 	{
@@ -1045,9 +1065,14 @@ void FSkeletalMeshGpuSpawnStaticBuffers::InitRHI(FRHICommandListBase& RHICmdList
 	// There is not section selection/culling in the interface so technically we could compute that array in the pipeline.
 	if ( bMeshValid )
 	{
-		FRHIResourceCreateInfo CreateInfo(TEXT("FSkeletalMeshGpuSpawnStaticBuffers"));
-		BufferTriangleMatricesOffsetRHI = RHICmdList.CreateBuffer(VertexCount * sizeof(uint32), BUF_Static | BUF_VertexBuffer | BUF_ShaderResource, 0, ERHIAccess::VertexOrIndexBuffer | ERHIAccess::SRVMask, CreateInfo);
-		uint32* MatricesOffsets = (uint32*)RHICmdList.LockBuffer(BufferTriangleMatricesOffsetRHI, 0, VertexCount * sizeof(uint32), RLM_WriteOnly);
+		const FRHIBufferCreateDesc CreateDesc =
+			FRHIBufferCreateDesc::CreateVertex(TEXT("FSkeletalMeshGpuSpawnStaticBuffers"), VertexCount * sizeof(uint32))
+			.AddUsage(EBufferUsageFlags::Static | EBufferUsageFlags::ShaderResource)
+			.SetInitialState(ERHIAccess::VertexOrIndexBuffer | ERHIAccess::SRVMask)
+			.SetInitActionInitializer();
+
+		TRHIBufferInitializer<uint32> Initializer = RHICmdList.CreateBufferInitializer(CreateDesc);
+
 		uint32 AccumulatedMatrixOffset = 0;
 		for (uint32 s = 0; s < SectionCount; ++s)
 		{
@@ -1056,11 +1081,12 @@ void FSkeletalMeshGpuSpawnStaticBuffers::InitRHI(FRHICommandListBase& RHICmdList
 			const uint32 SectionNumVertices = Section.NumVertices;
 			for (uint32 SectionVertex = 0; SectionVertex < SectionNumVertices; ++SectionVertex)
 			{
-				MatricesOffsets[SectionBaseVertexIndex + SectionVertex] = AccumulatedMatrixOffset;
+				Initializer[SectionBaseVertexIndex + SectionVertex] = AccumulatedMatrixOffset;
 			}
 			AccumulatedMatrixOffset += Section.BoneMap.Num();
 		}
-		RHICmdList.UnlockBuffer(BufferTriangleMatricesOffsetRHI);
+
+		BufferTriangleMatricesOffsetRHI = Initializer.Finalize();
 		BufferTriangleMatricesOffsetSRV = RHICmdList.CreateShaderResourceView(
 			BufferTriangleMatricesOffsetRHI, 
 			FRHIViewDesc::CreateBufferSRV()
@@ -1078,10 +1104,13 @@ void FSkeletalMeshGpuSpawnStaticBuffers::InitRHI(FRHICommandListBase& RHICmdList
 	// Create arrays for filtered bones / sockets
 	if ( FilteredAndUnfilteredBonesArray.Num() > 0 )
 	{
-		FRHIResourceCreateInfo CreateInfo(TEXT("FilteredAndUnfilteredBonesBuffer"));
-		CreateInfo.ResourceArray = &FilteredAndUnfilteredBonesArray;
+		const FRHIBufferCreateDesc CreateDesc =
+			FRHIBufferCreateDesc::CreateVertex(TEXT("FilteredAndUnfilteredBonesBuffer"), FilteredAndUnfilteredBonesArray.Num() * FilteredAndUnfilteredBonesArray.GetTypeSize())
+			.AddUsage(EBufferUsageFlags::Static | EBufferUsageFlags::ShaderResource)
+			.SetInitActionResourceArray(&FilteredAndUnfilteredBonesArray)
+			.DetermineInitialState();
 
-		FilteredAndUnfilteredBonesBuffer = RHICmdList.CreateVertexBuffer(FilteredAndUnfilteredBonesArray.Num() * FilteredAndUnfilteredBonesArray.GetTypeSize(), BUF_Static | BUF_ShaderResource, CreateInfo);
+		FilteredAndUnfilteredBonesBuffer = RHICmdList.CreateBuffer(CreateDesc);
 		FilteredAndUnfilteredBonesSRV = RHICmdList.CreateShaderResourceView(
 			FilteredAndUnfilteredBonesBuffer,
 			FRHIViewDesc::CreateBufferSRV()
@@ -1163,15 +1192,24 @@ void FSkeletalMeshGpuDynamicBufferProxy::InitRHI(FRHICommandListBase& RHICmdList
 #endif
 	for (FSkeletalBuffer& Buffer : RWBufferBones)
 	{
-		FRHIResourceCreateInfo CreateInfo(TEXT("SkeletalMeshGpuDynamicBuffer"));
-		Buffer.SectionBuffer = RHICmdList.CreateVertexBuffer(sizeof(FVector4f) * 3 * SectionBoneCount, BUF_ShaderResource | BUF_Dynamic, CreateInfo);
+		const FRHIBufferCreateDesc SectionBufferDesc =
+			FRHIBufferCreateDesc::CreateVertex<FVector4f>(TEXT("SkeletalMeshGpuDynamicBuffer"), 3 * SectionBoneCount)
+			.AddUsage(EBufferUsageFlags::ShaderResource | EBufferUsageFlags::Dynamic)
+			.DetermineInitialState();
+
+		Buffer.SectionBuffer = RHICmdList.CreateBuffer(SectionBufferDesc);
 		Buffer.SectionSRV = RHICmdList.CreateShaderResourceView(
 			Buffer.SectionBuffer, 
 			FRHIViewDesc::CreateBufferSRV()
 				.SetType(FRHIViewDesc::EBufferType::Typed)
 				.SetFormat(PF_A32B32G32R32F));
 
-		Buffer.SamplingBuffer = RHICmdList.CreateVertexBuffer(sizeof(FVector4f) * 3 * (SamplingBoneCount + SamplingSocketCount), BUF_ShaderResource | BUF_Dynamic, CreateInfo);
+		const FRHIBufferCreateDesc SamplingBufferDesc =
+			FRHIBufferCreateDesc::CreateVertex<FVector4f>(TEXT("SkeletalMeshGpuDynamicBuffer"), 3 * (SamplingBoneCount + SamplingSocketCount))
+			.AddUsage(EBufferUsageFlags::ShaderResource | EBufferUsageFlags::Dynamic)
+			.DetermineInitialState();
+
+		Buffer.SamplingBuffer = RHICmdList.CreateBuffer(SamplingBufferDesc);
 		Buffer.SamplingSRV = RHICmdList.CreateShaderResourceView(
 			Buffer.SamplingBuffer, 
 			FRHIViewDesc::CreateBufferSRV()
diff --git a/Engine/Plugins/FX/Niagara/Source/Niagara/Private/NiagaraDataInterfaceSkeletalMeshConnectivity.cpp b/Engine/Plugins/FX/Niagara/Source/Niagara/Private/NiagaraDataInterfaceSkeletalMeshConnectivity.cpp
--- a/Engine/Plugins/FX/Niagara/Source/Niagara/Private/NiagaraDataInterfaceSkeletalMeshConnectivity.cpp
+++ b/Engine/Plugins/FX/Niagara/Source/Niagara/Private/NiagaraDataInterfaceSkeletalMeshConnectivity.cpp
@@ -400,12 +400,15 @@ FSkeletalMeshConnectivityProxy::Initialize(const FSkeletalMeshConnectivity& Conn
 void
 FSkeletalMeshConnectivityProxy::InitRHI(FRHICommandListBase& RHICmdList)
 {
-	FRHIResourceCreateInfo CreateInfo(TEXT("FSkeletalMeshConnectivityProxy_AdjacencyBuffer"));
-	CreateInfo.ResourceArray = &AdjacencyResource;
+	const int32 BufferSize = AdjacencyResource.GetResourceDataSize();
 
-	const int32 BufferSize = AdjacencyResource.Num();
+	const FRHIBufferCreateDesc CreateDesc =
+		FRHIBufferCreateDesc::CreateVertex(TEXT("FSkeletalMeshConnectivityProxy_AdjacencyBuffer"), BufferSize)
+		.AddUsage(EBufferUsageFlags::ShaderResource | EBufferUsageFlags::Static)
+		.SetInitActionResourceArray(&AdjacencyResource)
+		.DetermineInitialState();
 
-	AdjacencyBuffer = RHICmdList.CreateVertexBuffer(BufferSize, BUF_ShaderResource | BUF_Static, CreateInfo);
+	AdjacencyBuffer = RHICmdList.CreateBuffer(CreateDesc);
 	AdjacencySrv = RHICmdList.CreateShaderResourceView(
 		AdjacencyBuffer, 
 		FRHIViewDesc::CreateBufferSRV()
diff --git a/Engine/Plugins/FX/Niagara/Source/Niagara/Private/NiagaraEmptyUAVPool.cpp b/Engine/Plugins/FX/Niagara/Source/Niagara/Private/NiagaraEmptyUAVPool.cpp
--- a/Engine/Plugins/FX/Niagara/Source/Niagara/Private/NiagaraEmptyUAVPool.cpp
+++ b/Engine/Plugins/FX/Niagara/Source/Niagara/Private/NiagaraEmptyUAVPool.cpp
@@ -65,13 +65,16 @@ FRHIUnorderedAccessView* FNiagaraEmptyUAVPool::GetEmptyUAVFromPool(FRHICommandLi
 
 		// Initialize the UAV
 		const TCHAR* ResourceName = TEXT("FNiagaraGpuComputeDispatch::EmptyUAV");
-		FRHIResourceCreateInfo CreateInfo(ResourceName);
 		switch ( Type )
 		{
 			case ENiagaraEmptyUAVType::Buffer:
 			{
-				const uint32 BytesPerElement = GPixelFormats[Format].BlockBytes;
-				NewUAV.Buffer = RHICmdList.CreateVertexBuffer(BytesPerElement, BUF_UnorderedAccess | BUF_ShaderResource, CreateInfo);
+				const FRHIBufferCreateDesc CreateDesc =
+					FRHIBufferCreateDesc::CreateVertex(ResourceName, GPixelFormats[Format].BlockBytes)
+					.AddUsage(EBufferUsageFlags::UnorderedAccess | EBufferUsageFlags::ShaderResource)
+					.DetermineInitialState();
+
+				NewUAV.Buffer = RHICmdList.CreateBuffer(CreateDesc);
 				NewUAV.UAV = RHICmdList.CreateUnorderedAccessView(
 					NewUAV.Buffer, 
 					FRHIViewDesc::CreateBufferUAV()
diff --git a/Engine/Plugins/FX/Niagara/Source/Niagara/Private/NiagaraRenderer.cpp b/Engine/Plugins/FX/Niagara/Source/Niagara/Private/NiagaraRenderer.cpp
--- a/Engine/Plugins/FX/Niagara/Source/Niagara/Private/NiagaraRenderer.cpp
+++ b/Engine/Plugins/FX/Niagara/Source/Niagara/Private/NiagaraRenderer.cpp
@@ -55,19 +55,22 @@ public:
 	{
 		// Create a buffer with one element.
 		uint32 NumBytes = GPixelFormats[PixelFormat].BlockBytes;
-		FRHIResourceCreateInfo CreateInfo(*DebugName);
-		Buffer = RHICmdList.CreateVertexBuffer(NumBytes, BUF_ShaderResource | BUF_Static, CreateInfo);
 
-		// Zero the buffer memory.
-		void* Data = RHICmdList.LockBuffer(Buffer, 0, NumBytes, RLM_WriteOnly);
-		FMemory::Memset(Data, 0, NumBytes);
+		FRHIBufferCreateDesc CreateDesc =
+			FRHIBufferCreateDesc::CreateVertex(*DebugName, NumBytes)
+			.AddUsage(EBufferUsageFlags::ShaderResource | EBufferUsageFlags::Static)
+			.DetermineInitialState();
 
 		if (PixelFormat == PF_R8G8B8A8)
 		{
-			*reinterpret_cast<uint32*>(Data) = DefaultValue;
+			FRHIBufferInitializer Initializer = RHICmdList.CreateBufferInitializer(CreateDesc.SetInitActionInitializer());
+			Initializer.WriteData(&DefaultValue, sizeof(DefaultValue));
+			Buffer = Initializer.Finalize();
+		}
+		else
+		{
+			Buffer = RHICmdList.CreateBuffer(CreateDesc.SetInitActionZeroData());
 		}
-
-		RHICmdList.UnlockBuffer(Buffer);
 
 		SRV = RHICmdList.CreateShaderResourceView(
 			Buffer, 
@@ -108,7 +111,6 @@ public:
 			.SetFlags(ETextureCreateFlags::ShaderResource);
 
 		// Create a 1x1 texture.
-		FRHIResourceCreateInfo CreateInfo(*DebugName);
 
 		uint32 Stride;
 		switch (Dimension)
diff --git a/Engine/Plugins/FX/Niagara/Source/Niagara/Private/NiagaraRendererRibbons.cpp b/Engine/Plugins/FX/Niagara/Source/Niagara/Private/NiagaraRendererRibbons.cpp
--- a/Engine/Plugins/FX/Niagara/Source/Niagara/Private/NiagaraRendererRibbons.cpp
+++ b/Engine/Plugins/FX/Niagara/Source/Niagara/Private/NiagaraRendererRibbons.cpp
@@ -303,8 +303,12 @@ struct FNiagaraRibbonIndexBuffer final : FIndexBuffer
 	{
 		InitResource(RHICmdList);
 
-		FRHIResourceCreateInfo CreateInfo(TEXT("NiagaraRibbonIndexBuffer"));
-		IndexBufferRHI = RHICmdList.CreateIndexBuffer(sizeof(uint32), sizeof(uint32) * NumElements, BUF_Static | BUF_UnorderedAccess | BUF_ShaderResource, ERHIAccess::VertexOrIndexBuffer, CreateInfo);
+		const FRHIBufferCreateDesc CreateDesc =
+			FRHIBufferCreateDesc::CreateIndex<uint32>(TEXT("NiagaraRibbonIndexBuffer"), NumElements)
+			.AddUsage(EBufferUsageFlags::Static | EBufferUsageFlags::UnorderedAccess | EBufferUsageFlags::ShaderResource)
+			.SetInitialState(ERHIAccess::VertexOrIndexBuffer);
+
+		IndexBufferRHI = RHICmdList.CreateBuffer(CreateDesc);
 	#if RHI_RAYTRACING
 		SRV = RHICmdList.CreateShaderResourceView(
 			IndexBufferRHI, 
@@ -464,12 +468,14 @@ bool FNiagaraRibbonGpuBuffer::Allocate(FRHICommandListBase& RHICmdList, uint32 N
 		if (bGpuUsageChanged || CurrentElements < NumElements || CurrentElements > uint32(FMath::CeilToInt32(NumElements * DownsizeMultiplier)))
 		{
 			NumElements = FMath::Min(MaxElements, uint32(FMath::RoundToInt32(NumElements * UpsizeMultipler)));
+			NumBytes = ElementBytes * NumElements;
 
-			FRHIResourceCreateInfo CreateInfo(DebugName);
-			const EBufferUsageFlags UsageFlags = AdditionalBufferUsage | BUF_ShaderResource | (bGpuReadOnly ? BUF_Volatile : BUF_Static | BUF_UnorderedAccess);
+			const FRHIBufferCreateDesc CreateDesc =
+				FRHIBufferCreateDesc::CreateVertex(DebugName, NumBytes)
+				.AddUsage(AdditionalBufferUsage | EBufferUsageFlags::ShaderResource | (bGpuReadOnly ? EBufferUsageFlags::Volatile : EBufferUsageFlags::Static | EBufferUsageFlags::UnorderedAccess))
+				.SetInitialState(InResourceState);
 
-			NumBytes = ElementBytes * NumElements;
-			Buffer = RHICmdList.CreateVertexBuffer(NumBytes, UsageFlags, InResourceState, CreateInfo);
+			Buffer = RHICmdList.CreateBuffer(CreateDesc);
 			SRV = RHICmdList.CreateShaderResourceView(
 				Buffer, 
 				FRHIViewDesc::CreateBufferSRV()
@@ -814,19 +820,21 @@ void FNiagaraRendererRibbons::CreateRenderThreadResources(FRHICommandListBase& R
 	{
 		// Initialize the shape vertex buffer. This doesn't change frame-to-frame, so we can set it up once
 		const int32 NumElements = ShapeState.SliceTriangleToVertexIds.Num();
-		ShapeState.SliceTriangleToVertexIdsBuffer.Initialize(RHICmdList, TEXT("SliceTriangleToVertexIdsBuffer"), sizeof(uint32), NumElements, EPixelFormat::PF_R32_UINT, BUF_Static);
-		void* SliceTriangleToVertexIdsBufferPtr = RHICmdList.LockBuffer(ShapeState.SliceTriangleToVertexIdsBuffer.Buffer, 0, sizeof(uint32) * NumElements, RLM_WriteOnly);
-		FMemory::Memcpy(SliceTriangleToVertexIdsBufferPtr, ShapeState.SliceTriangleToVertexIds.GetData(), sizeof(uint32) * NumElements);
-		RHICmdList.UnlockBuffer(ShapeState.SliceTriangleToVertexIdsBuffer.Buffer);
+		ShapeState.SliceTriangleToVertexIdsBuffer.InitializeWithData(RHICmdList, TEXT("SliceTriangleToVertexIdsBuffer"), sizeof(uint32), NumElements, PF_R32_UINT, BUF_Static,
+			[this](FRHIBufferInitializer& Initializer)
+			{
+				Initializer.WriteData(ShapeState.SliceTriangleToVertexIds.GetData(), Initializer.GetWritableDataSize());
+			});
 	}
 
 	{
 		// Initialize the shape vertex buffer. This doesn't change frame-to-frame, so we can set it up once
 		const int32 NumElements = ShapeState.SliceVertexData.Num() * FNiagaraRibbonShapeGeometryData::FVertex::NumElements;
-		ShapeState.SliceVertexDataBuffer.Initialize(RHICmdList, TEXT("NiagaraShapeVertexDataBuffer"), sizeof(float), NumElements, EPixelFormat::PF_R32_FLOAT, BUF_Static);
-		void* SliceVertexDataBufferPtr = RHICmdList.LockBuffer(ShapeState.SliceVertexDataBuffer.Buffer, 0, sizeof(float) * NumElements, RLM_WriteOnly);
-		FMemory::Memcpy(SliceVertexDataBufferPtr, ShapeState.SliceVertexData.GetData(), sizeof(float) * NumElements);
-		RHICmdList.UnlockBuffer(ShapeState.SliceVertexDataBuffer.Buffer);
+		ShapeState.SliceVertexDataBuffer.InitializeWithData(RHICmdList, TEXT("NiagaraShapeVertexDataBuffer"), sizeof(float), NumElements, PF_R32_FLOAT, BUF_Static,
+			[this](FRHIBufferInitializer& Initializer)
+			{
+				Initializer.WriteData(ShapeState.SliceVertexData.GetData(), Initializer.GetWritableDataSize());
+			});
 	}
 	
 	
diff --git a/Engine/Plugins/FX/Niagara/Source/NiagaraVertexFactories/Private/NiagaraCutoutVertexBuffer.cpp b/Engine/Plugins/FX/Niagara/Source/NiagaraVertexFactories/Private/NiagaraCutoutVertexBuffer.cpp
--- a/Engine/Plugins/FX/Niagara/Source/NiagaraVertexFactories/Private/NiagaraCutoutVertexBuffer.cpp
+++ b/Engine/Plugins/FX/Niagara/Source/NiagaraVertexFactories/Private/NiagaraCutoutVertexBuffer.cpp
@@ -21,13 +21,16 @@ void FNiagaraCutoutVertexBuffer::InitRHI(FRHICommandListBase& RHICmdList)
 	if (Data.Num())
 	{
 		// create a static vertex buffer
-		FRHIResourceCreateInfo CreateInfo(TEXT("FNiagaraCutoutVertexBuffer"));
+		const FRHIBufferCreateDesc CreateDesc =
+			FRHIBufferCreateDesc::CreateVertex<FVector2f>(TEXT("FNiagaraCutoutVertexBuffer"), Data.Num())
+			.AddUsage(EBufferUsageFlags::Static | EBufferUsageFlags::ShaderResource)
+			.SetInitialState(ERHIAccess::VertexOrIndexBuffer | ERHIAccess::SRVMask)
+			.SetInitActionInitializer();
 
-		const int32 DataSize = Data.Num() * sizeof(FVector2f);
-		VertexBufferRHI = RHICmdList.CreateBuffer(DataSize, BUF_Static | BUF_VertexBuffer | BUF_ShaderResource, 0, ERHIAccess::VertexOrIndexBuffer | ERHIAccess::SRVMask, CreateInfo);
-		void* BufferData = RHICmdList.LockBuffer(VertexBufferRHI, 0, DataSize, RLM_WriteOnly);
-		FMemory::Memcpy(BufferData, Data.GetData(), DataSize);
-		RHICmdList.UnlockBuffer(VertexBufferRHI);
+		TRHIBufferInitializer<FVector2f> Initializer = RHICmdList.CreateBufferInitializer(CreateDesc);
+		Initializer.WriteArray(MakeConstArrayView(Data));
+
+		VertexBufferRHI = Initializer.Finalize();
 		VertexBufferSRV = RHICmdList.CreateShaderResourceView(
 			VertexBufferRHI, 
 			FRHIViewDesc::CreateBufferSRV()
diff --git a/Engine/Plugins/FX/Niagara/Source/NiagaraVertexFactories/Public/NiagaraVertexFactory.h b/Engine/Plugins/FX/Niagara/Source/NiagaraVertexFactories/Public/NiagaraVertexFactory.h
--- a/Engine/Plugins/FX/Niagara/Source/NiagaraVertexFactories/Public/NiagaraVertexFactory.h
+++ b/Engine/Plugins/FX/Niagara/Source/NiagaraVertexFactories/Public/NiagaraVertexFactory.h
@@ -23,13 +23,13 @@ public:
 	*/
 	virtual void InitRHI(FRHICommandListBase& RHICmdList) override
 	{
-		// create a static vertex buffer
-		FRHIResourceCreateInfo CreateInfo(TEXT("FNiagaraNullSortedIndicesVertexBuffer"));
-		VertexBufferRHI = RHICmdList.CreateBuffer(sizeof(uint32), BUF_Static | BUF_VertexBuffer | BUF_ShaderResource, 0, ERHIAccess::VertexOrIndexBuffer | ERHIAccess::SRVMask, CreateInfo);
-		void* BufferData = RHICmdList.LockBuffer(VertexBufferRHI, 0, sizeof(uint32), RLM_WriteOnly);
-		FMemory::Memzero(BufferData, sizeof(uint32));
-		RHICmdList.UnlockBuffer(VertexBufferRHI);
+		const FRHIBufferCreateDesc CreateDesc =
+			FRHIBufferCreateDesc::CreateVertex<uint32>(TEXT("FNiagaraNullSortedIndicesVertexBuffer"), 1)
+			.AddUsage(EBufferUsageFlags::Static | EBufferUsageFlags::ShaderResource)
+			.SetInitialState(ERHIAccess::VertexOrIndexBuffer | ERHIAccess::SRVMask)
+			.SetInitActionZeroData();
 
+		VertexBufferRHI = RHICmdList.CreateBuffer(CreateDesc);
 		VertexBufferSRV = RHICmdList.CreateShaderResourceView(
 			VertexBufferRHI, 
 			FRHIViewDesc::CreateBufferSRV()
diff --git a/Engine/Plugins/Media/AvfMedia/Source/AvfMedia/Private/Player/AvfMediaVideoSampler.cpp b/Engine/Plugins/Media/AvfMedia/Source/AvfMedia/Private/Player/AvfMediaVideoSampler.cpp
--- a/Engine/Plugins/Media/AvfMedia/Source/AvfMedia/Private/Player/AvfMediaVideoSampler.cpp
+++ b/Engine/Plugins/Media/AvfMedia/Source/AvfMedia/Private/Player/AvfMediaVideoSampler.cpp
@@ -269,7 +269,7 @@ void FAvfMediaVideoSampler::ProcessFrame(CVPixelBufferRef Frame, FTimespan Sampl
 
 					SetShaderParametersLegacyPS(RHICmdList, PixelShader, YTex, UVTex, *ColorTransform, MediaShaders::YUVOffset8bits, true);
 
-					FBufferRHIRef VertexBuffer = CreateTempMediaVertexBuffer();
+					FBufferRHIRef VertexBuffer = CreateTempMediaVertexBuffer(RHICmdList);
 					RHICmdList.SetStreamSource(0, VertexBuffer, 0);
 					RHICmdList.SetViewport(0, 0, 0.0f, YWidth, YHeight, 1.0f);
 
diff --git a/Engine/Plugins/Media/ElectraUtil/Source/ElectraSamples/Private/ElectraTextureSampleAVF.cpp b/Engine/Plugins/Media/ElectraUtil/Source/ElectraSamples/Private/ElectraTextureSampleAVF.cpp
--- a/Engine/Plugins/Media/ElectraUtil/Source/ElectraSamples/Private/ElectraTextureSampleAVF.cpp
+++ b/Engine/Plugins/Media/ElectraUtil/Source/ElectraSamples/Private/ElectraTextureSampleAVF.cpp
@@ -22,10 +22,6 @@ THIRD_PARTY_INCLUDES_START
 #include "MetalInclude.h"
 THIRD_PARTY_INCLUDES_END
 
-#if WITH_ENGINE
-extern void SafeReleaseMetalObject(NS::Object* Object);
-#endif
-
 // ------------------------------------------------------------------------------------------------------
 // ------------------------------------------------------------------------------------------------------
 // ------------------------------------------------------------------------------------------------------
@@ -40,10 +36,16 @@ FElectraMediaTexConvApple::FElectraMediaTexConvApple()
 FElectraMediaTexConvApple::~FElectraMediaTexConvApple()
 {
 #if WITH_ENGINE
+	// TODO: Carl - Test This  
 	if (MetalTextureCache)
 	{
 		CVMetalTextureCacheRef TextureCacheCopy = MetalTextureCache;
-		SafeReleaseMetalObject((__bridge NS::Object*)TextureCacheCopy);
+		
+		FRHICommandListImmediate& RHICmdList = FRHICommandListImmediate::Get();
+		RHICmdList.EnqueueLambda([&TextureCacheCopy](FRHICommandListBase& RHICmdList)
+		{
+			TextureCacheCopy = nullptr;
+		});
 	}
 #endif
 }
@@ -300,7 +302,7 @@ void FElectraMediaTexConvApple::ConvertTexture(FTextureRHIRef & InDstTexture, CV
 					}
 
 
-					FBufferRHIRef VertexBuffer = CreateTempMediaVertexBuffer();
+					FBufferRHIRef VertexBuffer = CreateTempMediaVertexBuffer(RHICmdList);
 					RHICmdList.SetStreamSource(0, VertexBuffer, 0);
 					RHICmdList.SetViewport(0, 0, 0.0f, YWidth, YHeight, 1.0f);
 
diff --git a/Engine/Plugins/Media/ElectraUtil/Source/ElectraSamples/Private/ElectraTextureSampleAndroid.cpp b/Engine/Plugins/Media/ElectraUtil/Source/ElectraSamples/Private/ElectraTextureSampleAndroid.cpp
--- a/Engine/Plugins/Media/ElectraUtil/Source/ElectraSamples/Private/ElectraTextureSampleAndroid.cpp
+++ b/Engine/Plugins/Media/ElectraUtil/Source/ElectraSamples/Private/ElectraTextureSampleAndroid.cpp
@@ -851,7 +851,7 @@ void FElectraTextureSample::CopyFromExternalTextureOES(FRHICommandListImmediate&
 		SetShaderParametersLegacyPS(RHICmdList, CopyShader, InSrcTexture, InSamplerState, ScaleRotation, Offset);
 
 		// draw full size quad into render target
-		FBufferRHIRef VertexBuffer = CreateTempMediaVertexBuffer();
+		FBufferRHIRef VertexBuffer = CreateTempMediaVertexBuffer(RHICmdList);
 		RHICmdList.SetStreamSource(0, VertexBuffer, 0);
 		// set viewport to RT size
 		RHICmdList.SetViewport(0, 0, 0.0f, (float)OutputDim.X, (float)OutputDim.Y, 1.0f);
@@ -926,7 +926,7 @@ void  FElectraTextureSample::CopyFromExternalTextureVulkan(FRHICommandListImmedi
 
 
 		// draw full size quad into render target
-		FBufferRHIRef VertexBuffer = CreateTempMediaVertexBuffer(); 
+		FBufferRHIRef VertexBuffer = CreateTempMediaVertexBuffer(RHICmdList); 
 		RHICmdList.SetStreamSource(0, VertexBuffer, 0);
 		// set viewport to RT size
 		RHICmdList.SetViewport(0, 0, 0.0f, (float)OutputDim.X, (float)OutputDim.Y, 1.0f);
@@ -1044,7 +1044,7 @@ bool FElectraTextureSample::ConvertCpuOutputPath(FRHICommandListImmediate& RHICm
 	RHICmdList.Transition(FRHITransitionInfo(InDstTexture, ERHIAccess::Unknown, ERHIAccess::RTV));
 	FRHIRenderPassInfo RPInfo(InDstTexture, ERenderTargetActions::DontLoad_Store);
 
-	FBufferRHIRef TempVB = CreateTempMediaVertexBuffer();
+	FBufferRHIRef TempVB = CreateTempMediaVertexBuffer(RHICmdList);
 
 	RHICmdList.BeginRenderPass(RPInfo, TEXT("AndroidProcessVideo"));
 
diff --git a/Engine/Plugins/Media/ImgMedia/Source/ImgMedia/Private/Readers/ExrImgMediaReaderGpu.cpp b/Engine/Plugins/Media/ImgMedia/Source/ImgMedia/Private/Readers/ExrImgMediaReaderGpu.cpp
--- a/Engine/Plugins/Media/ImgMedia/Source/ImgMedia/Private/Readers/ExrImgMediaReaderGpu.cpp
+++ b/Engine/Plugins/Media/ImgMedia/Source/ImgMedia/Private/Readers/ExrImgMediaReaderGpu.cpp
@@ -505,16 +505,18 @@ void FExrImgMediaReaderGpu::CreateSampleConverterCallback(TSharedPtr<FExrMediaTe
 			{
 				TRACE_CPUPROFILER_EVENT_SCOPE_STR("ExrReaderGpu.TileDesc");
 
-				FBufferRHIRef BufferRef;
-				FRHIResourceCreateInfo CreateInfo(TEXT("FExrImgMediaReaderGpu_TileDesc"));
-				const uint32 BytesPerElement = sizeof(FExrReader::FTileDesc);
-				const uint32 NumElements = ConverterParams.TileInfoPerMipLevel[SampleMipLevel].Num();
-
 				// This buffer is allocated on already allocated block, therefore the risk of fragmentation is mitigated.
-				BufferRef = RHICmdList.CreateStructuredBuffer(BytesPerElement, BytesPerElement * NumElements, BUF_ShaderResource | BUF_Dynamic | BUF_FastVRAM, CreateInfo);
-				void* MappedBuffer = RHICmdList.LockBuffer(BufferRef, 0, NumElements, RLM_WriteOnly);
-				FMemory::Memcpy(MappedBuffer, &ConverterParams.TileInfoPerMipLevel[SampleMipLevel][0], sizeof(FExrReader::FTileDesc) * ConverterParams.TileInfoPerMipLevel[SampleMipLevel].Num());
-				RHICmdList.UnlockBuffer(BufferRef);
+				const FRHIBufferCreateDesc CreateDesc =
+					FRHIBufferCreateDesc::CreateStructured<FExrReader::FTileDesc>(TEXT("FExrImgMediaReaderGpu_TileDesc"), ConverterParams.TileInfoPerMipLevel[SampleMipLevel].Num())
+					.AddUsage(EBufferUsageFlags::ShaderResource | EBufferUsageFlags::Dynamic | EBufferUsageFlags::FastVRAM)
+					.SetInitActionInitializer()
+					.DetermineInitialState();
+
+				TRHIBufferInitializer<FExrReader::FTileDesc> Initializer = RHICmdList.CreateBufferInitializer(CreateDesc);
+				Initializer.WriteArray(MakeConstArrayView(ConverterParams.TileInfoPerMipLevel[SampleMipLevel]));
+
+				FBufferRHIRef BufferRef = Initializer.Finalize();
+
 				Parameters.TileDescBuffer = RHICmdList.CreateShaderResourceView(BufferRef, FRHIViewDesc::CreateBufferSRV().SetTypeFromBuffer(BufferRef));
 				PermutationVector.Set<FExrSwizzlePS::FPartialTiles>(true);
 			}
@@ -652,15 +654,23 @@ FStructuredBufferPoolItemSharedPtr FExrImgMediaReaderGpu::AllocateGpuBufferFromP
 				RHI_BREADCRUMB_EVENT_STAT(RHICmdList, ExrImgMediaReaderGpu_AllocateBuffer, "ExrImgGpu.MipRender.AllocateBuffer");
 				SCOPED_GPU_STAT(RHICmdList, ExrImgMediaReaderGpu_AllocateBuffer);
 
-				FRHIResourceCreateInfo CreateInfo(TEXT(""));
-				CreateInfo.DebugName = TEXT("ExrReaderGpu.UploadBuffer");
-				AllocatedBuffer->UploadBufferRef = RHICmdList.CreateStructuredBuffer(sizeof(uint16) * 2., AllocSize, BUF_ShaderResource | BUF_Dynamic | BUF_FastVRAM, CreateInfo);
-				AllocatedBuffer->UploadBufferMapped = RHICmdList.LockBuffer(AllocatedBuffer->UploadBufferRef, 0, AllocSize, RLM_WriteOnly);
+				{
+					const FRHIBufferCreateDesc CreateDesc =
+						FRHIBufferCreateDesc::CreateStructured(TEXT("ExrReaderGpu.UploadBuffer"), AllocSize, sizeof(uint16) * 2)
+						.AddUsage(EBufferUsageFlags::ShaderResource | EBufferUsageFlags::Dynamic | EBufferUsageFlags::FastVRAM)
+						.DetermineInitialState();
+					AllocatedBuffer->UploadBufferRef = RHICmdList.CreateBuffer(CreateDesc);
+					AllocatedBuffer->UploadBufferMapped = RHICmdList.LockBuffer(AllocatedBuffer->UploadBufferRef, 0, AllocSize, RLM_WriteOnly);
+				}
 
 				if (bExrReaderUseUploadHeap)
 				{
-					CreateInfo.DebugName = TEXT("ExrReaderGpu.DestBuffer");
-					AllocatedBuffer->ShaderAccessBufferRef = RHICmdList.CreateStructuredBuffer(sizeof(uint16) * 2., AllocSize, BUF_ShaderResource | BUF_FastVRAM, CreateInfo);
+					const FRHIBufferCreateDesc CreateDesc =
+						FRHIBufferCreateDesc::CreateStructured(TEXT("ExrReaderGpu.DestBuffer"), AllocSize, sizeof(uint16) * 2)
+						.AddUsage(EBufferUsageFlags::ShaderResource | EBufferUsageFlags::FastVRAM)
+						.DetermineInitialState();
+
+					AllocatedBuffer->ShaderAccessBufferRef = RHICmdList.CreateBuffer(CreateDesc);
 					AllocatedBuffer->ShaderResourceView = RHICmdList.CreateShaderResourceView(AllocatedBuffer->ShaderAccessBufferRef, FRHIViewDesc::CreateBufferSRV().SetTypeFromBuffer(AllocatedBuffer->ShaderAccessBufferRef));
 				}
 				else
diff --git a/Engine/Plugins/Runtime/AR/AppleAR/AppleARKit/Source/AppleARKit/Private/AppleARKitTextures.cpp b/Engine/Plugins/Runtime/AR/AppleAR/AppleARKit/Source/AppleARKit/Private/AppleARKitTextures.cpp
--- a/Engine/Plugins/Runtime/AR/AppleAR/AppleARKit/Source/AppleARKit/Private/AppleARKitTextures.cpp
+++ b/Engine/Plugins/Runtime/AR/AppleAR/AppleARKit/Source/AppleARKit/Private/AppleARKitTextures.cpp
@@ -580,8 +580,6 @@ public:
 	 */
 	virtual void InitRHI(FRHICommandListBase& RHICmdList) override
 	{
-		FRHIResourceCreateInfo CreateInfo(TEXT("FARMetalResource"));
-
 		id<MTLTexture> MetalTexture = Owner->GetMetalTexture();
 		if (MetalTexture != nullptr)
 		{
diff --git a/Engine/Plugins/Runtime/CableComponent/Source/CableComponent/Private/CableComponent.cpp b/Engine/Plugins/Runtime/CableComponent/Source/CableComponent/Private/CableComponent.cpp
--- a/Engine/Plugins/Runtime/CableComponent/Source/CableComponent/Private/CableComponent.cpp
+++ b/Engine/Plugins/Runtime/CableComponent/Source/CableComponent/Private/CableComponent.cpp
@@ -60,8 +60,12 @@ class FCableIndexBuffer : public FIndexBuffer
 public:
 	virtual void InitRHI(FRHICommandListBase& RHICmdList) override
 	{
-		FRHIResourceCreateInfo CreateInfo(TEXT("FCableIndexBuffer"));
-		IndexBufferRHI = RHICmdList.CreateIndexBuffer(sizeof(int32), NumIndices * sizeof(int32), BUF_Dynamic | BUF_ShaderResource, CreateInfo);
+		const FRHIBufferCreateDesc CreateDesc =
+			FRHIBufferCreateDesc::CreateIndex<int32>(TEXT("FCableIndexBuffer"), NumIndices)
+			.AddUsage(EBufferUsageFlags::Dynamic | EBufferUsageFlags::ShaderResource)
+			.DetermineInitialState();
+
+		IndexBufferRHI = RHICmdList.CreateBuffer(CreateDesc);
 	}
 
 	int32 NumIndices;
