[GPULM] Fix initial state of IrradianceCache buffers.

#rnx
#tests GPULM test maps

diff --git a/Engine/Plugins/Experimental/GPULightmass/Source/GPULightmass/Private/IrradianceCaching.cpp b/Engine/Plugins/Experimental/GPULightmass/Source/GPULightmass/Private/IrradianceCaching.cpp
--- a/Engine/Plugins/Experimental/GPULightmass/Source/GPULightmass/Private/IrradianceCaching.cpp
+++ b/Engine/Plugins/Experimental/GPULightmass/Source/GPULightmass/Private/IrradianceCaching.cpp
@@ -18,13 +18,15 @@ FIrradianceCache::FIrradianceCache(FRHICommandList& RHICmdList, int32 Quality, f
 	{
 		const FRHIBufferCreateDesc RecordsDesc =
 			FRHIBufferCreateDesc::CreateStructured<FVector4>(TEXT("FIrradianceCache"), IrradianceCacheMaxSize)
-			.AddUsage(BUF_UnorderedAccess | BUF_ShaderResource);
+			.AddUsage(BUF_UnorderedAccess | BUF_ShaderResource)
+			.DetermineInitialState();
 		IrradianceCacheRecords = RHICmdList.CreateBuffer(RecordsDesc);
 		IrradianceCacheRecordsUAV = RHICmdList.CreateUnorderedAccessView(IrradianceCacheRecords, FRHIViewDesc::CreateBufferUAV().SetTypeFromBuffer(IrradianceCacheRecords));
 
 		const FRHIBufferCreateDesc RecordBackfaceHitsDesc =
 			FRHIBufferCreateDesc::CreateStructured<FVector4>(TEXT("FIrradianceCache"), IrradianceCacheMaxSize)
-			.AddUsage(BUF_UnorderedAccess | BUF_ShaderResource);
+			.AddUsage(BUF_UnorderedAccess | BUF_ShaderResource)
+			.DetermineInitialState();
 		IrradianceCacheRecordBackfaceHits = RHICmdList.CreateBuffer(RecordBackfaceHitsDesc);
 		IrradianceCacheRecordBackfaceHitsUAV = RHICmdList.CreateUnorderedAccessView(IrradianceCacheRecordBackfaceHits, FRHIViewDesc::CreateBufferUAV().SetTypeFromBuffer(IrradianceCacheRecordBackfaceHits));
 		
