Fix bug in LightmapRenderer when using RequiresSeparateHitGroupContributionsBuffer

#DisabledValidations Preflight, CheckUnacceptableWords

diff --git a/Engine/Plugins/Experimental/GPULightmass/Source/GPULightmass/Private/LightmapRenderer.cpp b/Engine/Plugins/Experimental/GPULightmass/Source/GPULightmass/Private/LightmapRenderer.cpp
--- a/Engine/Plugins/Experimental/GPULightmass/Source/GPULightmass/Private/LightmapRenderer.cpp
+++ b/Engine/Plugins/Experimental/GPULightmass/Source/GPULightmass/Private/LightmapRenderer.cpp
@@ -1213,7 +1213,11 @@ bool FSceneRenderState::SetupRayTracingScene(FRDGBuilder& GraphBuilder, FSceneUn
 			InstanceBuffer.Initialize(RHICmdList, TEXT("LightmassRayTracingInstanceBuffer"), GRHIRayTracingInstanceDescriptorSize, SceneInitializer.MaxNumInstances);
 			
 			FRWBufferStructured HitGroupContributionsBuffer;
-			InstanceBuffer.Initialize(RHICmdList, TEXT("LightmassRayTracingHitGroupBuffer"), 4, SceneInitializer.MaxNumInstances);
+			
+			if(GRHIGlobals.RayTracing.RequiresSeparateHitGroupContributionsBuffer)
+			{
+				HitGroupContributionsBuffer.Initialize(RHICmdList, TEXT("LightmassRayTracingHitGroupBuffer"), 4, SceneInitializer.MaxNumInstances);
+			}
 
 			RayTracingInstanceBufferBuilder.FillRayTracingInstanceUploadBuffer(RHICmdList);
 			RayTracingInstanceBufferBuilder.FillAccelerationStructureAddressesBuffer(RHICmdList);
