Updated LightmapRenderer to also use MeshBatch.SegmentIndex when setting up bindings.

diff --git a/Engine/Plugins/Experimental/GPULightmass/Source/GPULightmass/Private/LightmapRenderer.cpp b/Engine/Plugins/Experimental/GPULightmass/Source/GPULightmass/Private/LightmapRenderer.cpp
--- a/Engine/Plugins/Experimental/GPULightmass/Source/GPULightmass/Private/LightmapRenderer.cpp
+++ b/Engine/Plugins/Experimental/GPULightmass/Source/GPULightmass/Private/LightmapRenderer.cpp
@@ -1147,12 +1147,12 @@ bool FSceneRenderState::SetupRayTracingScene(FRDGBuilder& GraphBuilder, FSceneUn
 						MeshBatches[0].Elements[0].MinVertexIndex = 0;
 						MeshBatches[0].Elements[0].MaxVertexIndex = 0;
 
-						for (int32 SegmentIndex = 0; SegmentIndex < MeshBatches.Num(); SegmentIndex++)
+						for (const FMeshBatch& MeshBatch : MeshBatches)
 						{
-							FDynamicRayTracingMeshCommandContext CommandContext(DynamicRayTracingMeshCommandStorage, DynamicRTShaderBindings, RayTracingInstance.GeometryRHI, SegmentIndex, SBTAllocation);
+							FDynamicRayTracingMeshCommandContext CommandContext(DynamicRayTracingMeshCommandStorage, DynamicRTShaderBindings, RayTracingInstance.GeometryRHI, MeshBatch.SegmentIndex, SBTAllocation);
 							FLightmapRayTracingMeshProcessor RayTracingMeshProcessor(&CommandContext);
 
-							RayTracingMeshProcessor.AddMeshBatch(MeshBatches[SegmentIndex], 1, nullptr);
+							RayTracingMeshProcessor.AddMeshBatch(MeshBatch, 1, nullptr);
 						}
 
 						FRayTracingMaskAndStatus RayTracingMaskAndStatus;
