Remove more uses of GRHICommandList.GetImmediateCommandList() and FRHICommandListExecutor::GetImmediateCommandList()
- Replaced with either FRHICommandListImmediate::Get(), or handing down a command list that was already accessible at the call sites.

#Synced-CL 48210367

diff --git a/Engine/Plugins/Developer/RenderDocPlugin/Source/RenderDocPlugin/Private/RenderDocPluginModule.cpp b/Engine/Plugins/Developer/RenderDocPlugin/Source/RenderDocPlugin/Private/RenderDocPluginModule.cpp
--- a/Engine/Plugins/Developer/RenderDocPlugin/Source/RenderDocPlugin/Private/RenderDocPluginModule.cpp
+++ b/Engine/Plugins/Developer/RenderDocPlugin/Source/RenderDocPlugin/Private/RenderDocPluginModule.cpp
@@ -130,9 +130,10 @@ public:
 		RenderDocAPI->StartFrameCapture(GetRenderdocDevicePointer(), WindowHandle);
 	}
 
-	static uint32 EndFrameCapture(HWND WindowHandle, FRenderDocPluginLoader::RENDERDOC_API_CONTEXT* RenderDocAPI)
+	static uint32 EndFrameCapture(FRHICommandListImmediate& RHICmdList, HWND WindowHandle, FRenderDocPluginLoader::RENDERDOC_API_CONTEXT* RenderDocAPI)
 	{
-		FRHICommandListExecutor::GetImmediateCommandList().SubmitCommandsAndFlushGPU();
+		RHICmdList.SubmitCommandsAndFlushGPU();
+
 		uint32 Result = RenderDocAPI->EndFrameCapture(GetRenderdocDevicePointer(), WindowHandle);
 		SetEmitDrawEvents(UE_GEmitDrawEvents_BeforeCapture);
 		return Result;
@@ -411,7 +412,7 @@ void FRenderDocPluginModule::EndFrameCapture(void* HWnd, uint32 Flags, const FSt
 	ENQUEUE_RENDER_COMMAND(EndRenderDocCapture)(
 		[WindowHandle, RenderDocAPILocal, Plugin, Flags, DestFileName](FRHICommandListImmediate& RHICmdList)
 		{
-			FRenderDocFrameCapturer::EndFrameCapture(WindowHandle, RenderDocAPILocal);
+			FRenderDocFrameCapturer::EndFrameCapture(RHICmdList, WindowHandle, RenderDocAPILocal);
 			FRenderDocFrameCapturer::SaveAndLaunch(Plugin, Flags, DestFileName);
 		});
 
diff --git a/Engine/Plugins/Developer/XcodeGPUDebuggerPlugin/Source/XcodeGPUDebuggerPlugin/Private/XcodeGPUDebuggerPluginModule.cpp b/Engine/Plugins/Developer/XcodeGPUDebuggerPlugin/Source/XcodeGPUDebuggerPlugin/Private/XcodeGPUDebuggerPluginModule.cpp
--- a/Engine/Plugins/Developer/XcodeGPUDebuggerPlugin/Source/XcodeGPUDebuggerPlugin/Private/XcodeGPUDebuggerPluginModule.cpp
+++ b/Engine/Plugins/Developer/XcodeGPUDebuggerPlugin/Source/XcodeGPUDebuggerPlugin/Private/XcodeGPUDebuggerPluginModule.cpp
@@ -90,9 +90,9 @@ public:
 		}
 	}
 
-	static void EndFrameCapture()
+	static void EndFrameCapture(FRHICommandListImmediate& RHICmdList)
 	{
-		FRHICommandListExecutor::GetImmediateCommandList().SubmitCommandsAndFlushGPU();
+		RHICmdList.SubmitCommandsAndFlushGPU();
 
 		MTL::CaptureManager* CaptureManager = MTL::CaptureManager::sharedCaptureManager();
 		CaptureManager->stopCapture();
@@ -230,7 +230,7 @@ void FXcodeGPUDebuggerPluginModule::EndFrameCapture(void* HWnd, uint32 Flags, co
 	ENQUEUE_RENDER_COMMAND(EndXcodeGPUDebuggerCapture)(
 		[Plugin = this, Flags, DestFileName](FRHICommandListImmediate& RHICmdList)
 		{
-			FXcodeGPUDebuggerFrameCapturer::EndFrameCapture();
+			FXcodeGPUDebuggerFrameCapturer::EndFrameCapture(RHICmdList);
 			FXcodeGPUDebuggerFrameCapturer::SaveAndLaunch(Plugin, Flags, DestFileName);
 		});
 
diff --git a/Engine/Plugins/Experimental/GPULightmass/Source/GPULightmass/Private/LightmapRenderer.cpp b/Engine/Plugins/Experimental/GPULightmass/Source/GPULightmass/Private/LightmapRenderer.cpp
--- a/Engine/Plugins/Experimental/GPULightmass/Source/GPULightmass/Private/LightmapRenderer.cpp
+++ b/Engine/Plugins/Experimental/GPULightmass/Source/GPULightmass/Private/LightmapRenderer.cpp
@@ -851,25 +851,21 @@ FCachedRayTracingSceneData::~FCachedRayTracingSceneData()
 		RaytracingSBT.FreeStaticRange(SBTAllocation);
 	}
 
-	FRHICommandListImmediate& RHICmdList = FRHICommandListExecutor::GetImmediateCommandList();
-
 	// Move OwnedRayTracingInstanceTransforms to RHIThread and it will get destroyed when the lambda is released
-	RHICmdList.EnqueueLambda([OwnedRayTracingInstanceTransforms = MoveTemp(OwnedRayTracingInstanceTransforms)](FRHICommandList&){});
+	FRHICommandListImmediate::Get().EnqueueLambda([OwnedRayTracingInstanceTransforms = MoveTemp(OwnedRayTracingInstanceTransforms)](FRHICommandList&){});
 }
 
 bool FSceneRenderState::SetupRayTracingScene(FRDGBuilder& GraphBuilder, FSceneUniformBuffer& SceneUniforms, int32 LODIndex)
 {
 	TRACE_CPUPROFILER_EVENT_SCOPE(SetupRayTracingScene);
 
-	FRHICommandListImmediate& RHICmdList = FRHICommandListExecutor::GetImmediateCommandList();
-
 	// Make sure the buffer is prepped
 	SceneUniforms.GetBuffer(GraphBuilder);
 
 #if RHI_RAYTRACING
 	// Force build all the open build requests
 	bool bBuildAll = true;
-	GRayTracingGeometryManager->ProcessBuildRequests(RHICmdList, bBuildAll);
+	GRayTracingGeometryManager->ProcessBuildRequests(GraphBuilder.RHICmdList, bBuildAll);
 #endif // RHI_RAYTRACING
 
 	if (CachedRayTracingScene.IsValid())
@@ -983,7 +979,7 @@ bool FSceneRenderState::SetupRayTracingScene(FRDGBuilder& GraphBuilder, FSceneUn
 	{
 		TRACE_CPUPROFILER_EVENT_SCOPE(RayTracingScene);
 
-		SCOPED_DRAW_EVENTF(RHICmdList, GPULightmassUpdateRayTracingScene, TEXT("GPULightmass UpdateRayTracingScene %d Instances"), StaticMeshInstanceRenderStates.Elements.Num());
+		SCOPED_DRAW_EVENTF(GraphBuilder.RHICmdList, GPULightmassUpdateRayTracingScene, TEXT("GPULightmass UpdateRayTracingScene %d Instances"), StaticMeshInstanceRenderStates.Elements.Num());
 
 		TArray<FRayTracingGeometryInstance, SceneRenderingAllocator> RayTracingGeometryInstances;
 		if (LODIndex != INDEX_NONE)
@@ -1092,7 +1088,7 @@ bool FSceneRenderState::SetupRayTracingScene(FRDGBuilder& GraphBuilder, FSceneUn
 							};
 
 							DynamicGeometryUpdateManager.AddDynamicGeometryToUpdate(
-								RHICmdList,
+								GraphBuilder.RHICmdList,
 								Landscape.ComponentUObject->GetWorld()->Scene->GetRenderScene(),
 								ReferenceView.Get(),
 								nullptr,
@@ -1168,7 +1164,7 @@ bool FSceneRenderState::SetupRayTracingScene(FRDGBuilder& GraphBuilder, FSceneUn
 
 		if (IsRayTracingEnabled())
 		{
-			SCOPED_GPU_MASK(RHICmdList, FRHIGPUMask::All());
+			SCOPED_GPU_MASK(GraphBuilder.RHICmdList, FRHIGPUMask::All());
 
 			FRayTracingInstanceBufferBuilder RayTracingInstanceBufferBuilder;
 			RayTracingInstanceBufferBuilder.Init(RayTracingGeometryInstances, View.ViewMatrices.GetPreViewTranslation());
@@ -1190,30 +1186,30 @@ bool FSceneRenderState::SetupRayTracingScene(FRDGBuilder& GraphBuilder, FSceneUn
 				FRHIBufferCreateDesc::Create(TEXT("LightmassRayTracingSceneBuffer"), SizeInfo.ResultSize, 0, EBufferUsageFlags::AccelerationStructure)
 				.SetInitialState(ERHIAccess::BVHWrite);
 
-			RayTracingSceneBuffer = RHICmdList.CreateBuffer(SceneCreateDesc);
-			RayTracingSceneSRV = RHICmdList.CreateShaderResourceView(FShaderResourceViewInitializer(RayTracingSceneBuffer, RayTracingScene, 0));
+			RayTracingSceneBuffer = GraphBuilder.RHICmdList.CreateBuffer(SceneCreateDesc);
+			RayTracingSceneSRV = GraphBuilder.RHICmdList.CreateShaderResourceView(FShaderResourceViewInitializer(RayTracingSceneBuffer, RayTracingScene, 0));
 
 			const FRHIBufferCreateDesc ScratchCreateDesc =
 				FRHIBufferCreateDesc::Create(TEXT("LightmassRayTracingScratchBuffer"), SizeInfo.BuildScratchSize, GRHIRayTracingScratchBufferAlignment, EBufferUsageFlags::StructuredBuffer | EBufferUsageFlags::RayTracingScratch)
 				.SetInitialState(ERHIAccess::UAVCompute);
 
-			FBufferRHIRef ScratchBuffer = RHICmdList.CreateBuffer(ScratchCreateDesc);
+			FBufferRHIRef ScratchBuffer = GraphBuilder.RHICmdList.CreateBuffer(ScratchCreateDesc);
 
 			FRWBufferStructured InstanceBuffer;
-			InstanceBuffer.Initialize(RHICmdList, TEXT("LightmassRayTracingInstanceBuffer"), GRHIRayTracingInstanceDescriptorSize, SceneInitializer.MaxNumInstances);
+			InstanceBuffer.Initialize(GraphBuilder.RHICmdList, TEXT("LightmassRayTracingInstanceBuffer"), GRHIRayTracingInstanceDescriptorSize, SceneInitializer.MaxNumInstances);
 			
 			FRWBufferStructured HitGroupContributionsBuffer;
 			
 			if(GRHIGlobals.RayTracing.RequiresSeparateHitGroupContributionsBuffer)
 			{
-				HitGroupContributionsBuffer.Initialize(RHICmdList, TEXT("LightmassRayTracingHitGroupBuffer"), 4, SceneInitializer.MaxNumInstances);
+				HitGroupContributionsBuffer.Initialize(GraphBuilder.RHICmdList, TEXT("LightmassRayTracingHitGroupBuffer"), 4, SceneInitializer.MaxNumInstances);
 			}
 
-			RayTracingInstanceBufferBuilder.FillRayTracingInstanceUploadBuffer(RHICmdList);
-			RayTracingInstanceBufferBuilder.FillAccelerationStructureAddressesBuffer(RHICmdList);
+			RayTracingInstanceBufferBuilder.FillRayTracingInstanceUploadBuffer(GraphBuilder.RHICmdList);
+			RayTracingInstanceBufferBuilder.FillAccelerationStructureAddressesBuffer(GraphBuilder.RHICmdList);
 
 			RayTracingInstanceBufferBuilder.BuildRayTracingInstanceBuffer(
-				RHICmdList,
+				GraphBuilder.RHICmdList,
 				nullptr,
 				nullptr,
 				InstanceBuffer.UAV,
@@ -1224,7 +1220,7 @@ bool FSceneRenderState::SetupRayTracingScene(FRDGBuilder& GraphBuilder, FSceneUn
 				0,
 				nullptr);
 
-			RHICmdList.BindAccelerationStructureMemory(RayTracingScene, RayTracingSceneBuffer, 0);
+			GraphBuilder.RHICmdList.BindAccelerationStructureMemory(RayTracingScene, RayTracingSceneBuffer, 0);
 
 			{
 				FRayTracingSceneBuildParams BuildParams;
@@ -1243,11 +1239,11 @@ bool FSceneRenderState::SetupRayTracingScene(FRDGBuilder& GraphBuilder, FSceneUn
 				BuildParams.ReferencedGeometries = RayTracingInstanceBufferBuilder.GetReferencedGeometries();
 				BuildParams.NumInstances = RayTracingInstanceBufferBuilder.GetMaxNumInstances();
 
-				RHICmdList.BuildAccelerationStructure(BuildParams);
+				GraphBuilder.RHICmdList.BuildAccelerationStructure(BuildParams);
 			}
 
 			// Move LandscapeTransforms to RHIThread to extend its lifetime until RHI cmd execution
-			RHICmdList.EnqueueLambda([LandscapeTransforms = MoveTemp(LandscapeTransforms)](FRHICommandList&){});
+			GraphBuilder.RHICmdList.EnqueueLambda([LandscapeTransforms = MoveTemp(LandscapeTransforms)](FRHICommandList&){});
 			
 			FRayTracingPipelineStateInitializer PSOInitializer;
 
@@ -1295,11 +1291,11 @@ bool FSceneRenderState::SetupRayTracingScene(FRDGBuilder& GraphBuilder, FSceneUn
 			FRHIRayTracingShader* MissTable[] = { GetGPULightmassDefaultMissShader(GlobalShaderMap) };
 			PSOInitializer.SetMissShaderTable(MissTable);
 
-			RayTracingPipelineState = PipelineStateCache::GetAndOrCreateRayTracingPipelineState(RHICmdList, PSOInitializer); 
+			RayTracingPipelineState = PipelineStateCache::GetAndOrCreateRayTracingPipelineState(GraphBuilder.RHICmdList, PSOInitializer);
 
 			RayTracingSBT.ResetStaticAllocationLock();
 
-			SBT = RayTracingSBT.AllocateTransientRHI(RHICmdList, ERayTracingShaderBindingMode::RTPSO, ERayTracingHitGroupIndexingMode::Allow, PSOInitializer.GetMaxLocalBindingDataSize());
+			SBT = RayTracingSBT.AllocateTransientRHI(GraphBuilder.RHICmdList, ERayTracingShaderBindingMode::RTPSO, ERayTracingHitGroupIndexingMode::Allow, PSOInitializer.GetMaxLocalBindingDataSize());
 
 			const int32 HiddenMaterialIndex = FindRayTracingHitGroupIndex(RayTracingPipelineState, HiddenMaterialShader, true);
 
@@ -1376,9 +1372,9 @@ bool FSceneRenderState::SetupRayTracingScene(FRDGBuilder& GraphBuilder, FSceneUn
 					FConcurrentLinearBulkObjectAllocator Allocator;
 
 					const uint32 MergedBindingsSize = sizeof(FRayTracingLocalShaderBindings) * NumTotalBindings;
-					FRayTracingLocalShaderBindings* MergedBindings = (FRayTracingLocalShaderBindings*)(RHICmdList.Bypass()
-						? Allocator.Malloc(MergedBindingsSize, alignof(FRayTracingLocalShaderBindings))
-						: RHICmdList.Alloc(MergedBindingsSize, alignof(FRayTracingLocalShaderBindings)));
+					FRayTracingLocalShaderBindings* MergedBindings = static_cast<FRayTracingLocalShaderBindings*>(GraphBuilder.RHICmdList.Bypass()
+						?              Allocator.Malloc(MergedBindingsSize, alignof(FRayTracingLocalShaderBindings))
+						: GraphBuilder.RHICmdList.Alloc(MergedBindingsSize, alignof(FRayTracingLocalShaderBindings)));
 
 					uint32 MergedBindingIndex = 0;
 					for (const FRayTracingLocalShaderBindingWriter::FChunk* Chunk = BindingWriter->GetFirstChunk(); Chunk; Chunk = Chunk->Next)
@@ -1392,7 +1388,7 @@ bool FSceneRenderState::SetupRayTracingScene(FRDGBuilder& GraphBuilder, FSceneUn
 					}
 
 					const bool bCopyDataToInlineStorage = false; // Storage is already allocated from RHICmdList, no extra copy necessary
-					RHICmdList.SetRayTracingHitGroups(
+					GraphBuilder.RHICmdList.SetRayTracingHitGroups(
 						SBT,
 						RayTracingPipelineState,
 						NumTotalBindings, MergedBindings,
@@ -1400,13 +1396,13 @@ bool FSceneRenderState::SetupRayTracingScene(FRDGBuilder& GraphBuilder, FSceneUn
 				}
 
 				// there is only one miss shader, so it must be at index 0 by definition
-				RHICmdList.SetRayTracingMissShader(SBT, 0, RayTracingPipelineState, 0 /* ShaderIndexInPipeline */, 0, nullptr, 0);
-				RHICmdList.CommitShaderBindingTable(SBT);
+				GraphBuilder.RHICmdList.SetRayTracingMissShader(SBT, 0, RayTracingPipelineState, 0 /* ShaderIndexInPipeline */, 0, nullptr, 0);
+				GraphBuilder.RHICmdList.CommitShaderBindingTable(SBT);
 
 				// Move the ray tracing binding container ownership to the command list, so that memory will be
 				// released on the RHI thread timeline, after the commands that reference it are processed.
 				// TUniquePtr<> auto destructs when exiting the lambda
-				RHICmdList.EnqueueLambda([BindingWriter = MoveTemp(BindingWriter)](FRHICommandListImmediate&) {});
+				GraphBuilder.RHICmdList.EnqueueLambda([BindingWriter = MoveTemp(BindingWriter)](FRHICommandListImmediate&) {});
 			}
 		}
 	}
@@ -1640,7 +1636,7 @@ void FLightmapRenderer::Finalize(FRDGBuilder& GraphBuilder)
 			int32 NewSize = FMath::CeilToInt(FMath::Sqrt(static_cast<float>(TileUploadRequests.Num())));
 			if (!UploadTilePoolGPU.IsValid() || UploadTilePoolGPU->SizeInTiles.X < NewSize)
 			{
-				UploadTilePoolGPU = MakeUnique<FLightmapTilePoolGPU>(4, FIntPoint(NewSize, NewSize), FIntPoint(GPreviewLightmapPhysicalTileSize, GPreviewLightmapPhysicalTileSize));
+				UploadTilePoolGPU = MakeUnique<FLightmapTilePoolGPU>(GraphBuilder.RHICmdList, 4, FIntPoint(NewSize, NewSize), FIntPoint(GPreviewLightmapPhysicalTileSize, GPreviewLightmapPhysicalTileSize));
 				UE_LOG(LogGPULightmass, Log, TEXT("Resizing GPULightmass upload tile pool to (%d, %d) %dx%d"), NewSize, NewSize, NewSize * GPreviewLightmapPhysicalTileSize, NewSize * GPreviewLightmapPhysicalTileSize);
 			}
 
@@ -2022,7 +2018,7 @@ void FLightmapRenderer::Finalize(FRDGBuilder& GraphBuilder)
 					}
 
 					int32 NewSize = FMath::Min(FMath::CeilToInt(FMath::Sqrt(static_cast<float>(TilesToQuery.Num()))), 64);
-					ScratchTilePoolGPU = MakeUnique<FLightmapTilePoolGPU>(3, FIntPoint(NewSize, NewSize), FIntPoint(GPreviewLightmapPhysicalTileSize, GPreviewLightmapPhysicalTileSize));
+					ScratchTilePoolGPU = MakeUnique<FLightmapTilePoolGPU>(GraphBuilder.RHICmdList, 3, FIntPoint(NewSize, NewSize), FIntPoint(GPreviewLightmapPhysicalTileSize, GPreviewLightmapPhysicalTileSize));
 					UE_LOG(LogGPULightmass, Log, TEXT("Resizing GPULightmass scratch tile pool to (%d, %d) %dx%d"), NewSize, NewSize, NewSize * GPreviewLightmapPhysicalTileSize, NewSize * GPreviewLightmapPhysicalTileSize);
 				}
 			}
@@ -3105,7 +3101,7 @@ void FLightmapRenderer::Finalize(FRDGBuilder& GraphBuilder)
 				LightmapReadbackGroup.ConvergedTileRequests = ConvergedTileRequestsOnCurrentGPU;
 				if (!LightmapReadbackGroup.ReadbackTilePoolGPU.IsValid())
 				{
-					LightmapReadbackGroup.ReadbackTilePoolGPU = MakeUnique<FLightmapTilePoolGPU>(4, FIntPoint(NewSize, NewSize), FIntPoint(GPreviewLightmapPhysicalTileSize, GPreviewLightmapPhysicalTileSize));
+					LightmapReadbackGroup.ReadbackTilePoolGPU = MakeUnique<FLightmapTilePoolGPU>(GraphBuilder.RHICmdList, 4, FIntPoint(NewSize, NewSize), FIntPoint(GPreviewLightmapPhysicalTileSize, GPreviewLightmapPhysicalTileSize));
 					LightmapReadbackGroup.StagingHQLayer0Readback = MakeUnique<FRHIGPUTextureReadback>(TEXT("StagingHQLayer0Readback"));
 					LightmapReadbackGroup.StagingHQLayer1Readback = MakeUnique<FRHIGPUTextureReadback>(TEXT("StagingHQLayer1Readback"));
 					LightmapReadbackGroup.StagingShadowMaskReadback = MakeUnique<FRHIGPUTextureReadback>(TEXT("StagingShadowMaskReadback"));
@@ -3280,7 +3276,7 @@ void FSceneRenderState::BuildMortonSortedLightmapRefList()
 	}
 }
 
-void FLightmapRenderer::BackgroundTick()
+void FLightmapRenderer::BackgroundTick(FRHICommandListImmediate& RHICmdList)
 {
 	{
 		TArray<FLightmapTileDenoiseGroup> FilteredDenoiseGroups;
@@ -3384,7 +3380,6 @@ void FLightmapRenderer::BackgroundTick()
 			ReadbackGroup.StagingShadowMaskReadback->IsReady(FRHIGPUMask::FromIndex(ReadbackGroup.GPUIndex)) &&
 			ReadbackGroup.StagingSkyOcclusionReadback->IsReady(FRHIGPUMask::FromIndex(ReadbackGroup.GPUIndex)))
 		{
-			FRHICommandListImmediate& RHICmdList = FRHICommandListExecutor::GetImmediateCommandList();
 			SCOPED_GPU_MASK(RHICmdList, FRHIGPUMask::FromIndex(ReadbackGroup.GPUIndex));
 
 			ReadbackGroup.TextureData = MakeUnique<FLightmapReadbackGroup::FTextureData>();
@@ -3666,8 +3661,6 @@ void FLightmapRenderer::BackgroundTick()
 		UE_LOG(LogGPULightmass, Log, TEXT("GPULightmass is now throttled for realtime preview"));
 	}
 
-	FRHICommandListImmediate& RHICmdList = FRHICommandListExecutor::GetImmediateCommandList();
-
 	if (!bOnlyBakeWhatYouSee)
 	{
 		const int32 NumWorkPerFrame = !bIsViewportNonRealtime ? 32 : 128;
diff --git a/Engine/Plugins/Experimental/GPULightmass/Source/GPULightmass/Private/LightmapRenderer.h b/Engine/Plugins/Experimental/GPULightmass/Source/GPULightmass/Private/LightmapRenderer.h
--- a/Engine/Plugins/Experimental/GPULightmass/Source/GPULightmass/Private/LightmapRenderer.h
+++ b/Engine/Plugins/Experimental/GPULightmass/Source/GPULightmass/Private/LightmapRenderer.h
@@ -127,7 +127,7 @@ public:
 
 	virtual ~FLightmapRenderer();
 
-	void BackgroundTick();
+	void BackgroundTick(FRHICommandListImmediate& RHICmdList);
 
 	void BumpRevision();
 	int32 GetCurrentRevision() { return CurrentRevision; }
diff --git a/Engine/Plugins/Experimental/GPULightmass/Source/GPULightmass/Private/LightmapTilePool.cpp b/Engine/Plugins/Experimental/GPULightmass/Source/GPULightmass/Private/LightmapTilePool.cpp
--- a/Engine/Plugins/Experimental/GPULightmass/Source/GPULightmass/Private/LightmapTilePool.cpp
+++ b/Engine/Plugins/Experimental/GPULightmass/Source/GPULightmass/Private/LightmapTilePool.cpp
@@ -108,7 +108,7 @@ void FLightmapTilePool::UnmapAll()
 	}
 }
 
-FLightmapTilePoolGPU::FLightmapTilePoolGPU(int32 NumLayers, FIntPoint InSizeInTiles, FIntPoint TileSize)
+FLightmapTilePoolGPU::FLightmapTilePoolGPU(FRHICommandListBase& RHICmdList, int32 NumLayers, FIntPoint InSizeInTiles, FIntPoint TileSize)
 	: FLightmapTilePool(InSizeInTiles.X * InSizeInTiles.Y)
 	, SizeInTiles(InSizeInTiles)
 {
@@ -136,8 +136,6 @@ FLightmapTilePoolGPU::FLightmapTilePoolGPU(int32 NumLayers, FIntPoint InSizeInTi
 			TexCreate_ShaderResource | TexCreate_RenderTargetable | TexCreate_UAV,
 			false);
 
-		FRHICommandListImmediate& RHICmdList = FRHICommandListExecutor::GetImmediateCommandList();
-
 		PooledRenderTargetDebugNames[RenderTargetIndex] = FString::Printf(TEXT("LightmapTilePoolGPU%d"), RenderTargetIndex);
 		GRenderTargetPool.FindFreeElement(RHICmdList, Desc, PooledRenderTargets[RenderTargetIndex], *PooledRenderTargetDebugNames[RenderTargetIndex]);
 
@@ -149,7 +147,7 @@ FLightmapTilePoolGPU::FLightmapTilePoolGPU(int32 NumLayers, FIntPoint InSizeInTi
 	// UE_LOG(LogGPULightmass, Log, TEXT("LightmapTilePool created with %.2fMB"), MemorySize / 1024.0f / 1024.0f);
 }
 
-void FLightmapTilePoolGPU::Initialize(FRHICommandList& RHICmdList, TArray<FLayerFormatAndTileSize> Layers)
+void FLightmapTilePoolGPU::Initialize(FRHICommandListBase& RHICmdList, TArray<FLayerFormatAndTileSize> Layers)
 {
 	ensure(SizeInTiles.X == SizeInTiles.Y);
 
diff --git a/Engine/Plugins/Experimental/GPULightmass/Source/GPULightmass/Private/LightmapTilePool.h b/Engine/Plugins/Experimental/GPULightmass/Source/GPULightmass/Private/LightmapTilePool.h
--- a/Engine/Plugins/Experimental/GPULightmass/Source/GPULightmass/Private/LightmapTilePool.h
+++ b/Engine/Plugins/Experimental/GPULightmass/Source/GPULightmass/Private/LightmapTilePool.h
@@ -56,9 +56,9 @@ public:
 	};
 
 	FLightmapTilePoolGPU(FIntPoint InSizeInTiles) : FLightmapTilePool(InSizeInTiles.X * InSizeInTiles.Y), SizeInTiles(InSizeInTiles) {}
-	FLightmapTilePoolGPU(int32 NumLayers, FIntPoint InSizeInTiles, FIntPoint TileSize);
+	FLightmapTilePoolGPU(FRHICommandListBase& RHICmdList, int32 NumLayers, FIntPoint InSizeInTiles, FIntPoint TileSize);
 
-	void Initialize(FRHICommandList& RHICmdList, TArray<FLayerFormatAndTileSize> Layers);
+	void Initialize(FRHICommandListBase& RHICmdList, TArray<FLayerFormatAndTileSize> Layers);
 
 	FIntPoint SizeInTiles;
 
diff --git a/Engine/Plugins/Experimental/GPULightmass/Source/GPULightmass/Private/Scene/Scene.cpp b/Engine/Plugins/Experimental/GPULightmass/Source/GPULightmass/Private/Scene/Scene.cpp
--- a/Engine/Plugins/Experimental/GPULightmass/Source/GPULightmass/Private/Scene/Scene.cpp
+++ b/Engine/Plugins/Experimental/GPULightmass/Source/GPULightmass/Private/Scene/Scene.cpp
@@ -1552,9 +1552,11 @@ void FScene::BackgroundTick()
 
 				float TargetDetailCellSize = Settings->VolumetricLightmapDetailCellSize;
 				
-				ENQUEUE_RENDER_COMMAND(BackgroundTickRenderThread)([&RenderState = RenderState, TargetDetailCellSize](FRHICommandListImmediate&) mutable {
+				ENQUEUE_RENDER_COMMAND(BackgroundTickRenderThread)(
+					[&RenderState = RenderState, TargetDetailCellSize](FRHICommandListImmediate& RHICmdList) mutable
+				{
 					RenderState.VolumetricLightmapRenderer->TargetDetailCellSize = TargetDetailCellSize;
-					RenderState.VolumetricLightmapRenderer->VoxelizeScene();
+					RenderState.VolumetricLightmapRenderer->VoxelizeScene(RHICmdList);
 					RenderState.VolumetricLightmapRenderer->FrameNumber = 0;
 					RenderState.VolumetricLightmapRenderer->SamplesTaken = 0;
 				});
@@ -1563,7 +1565,7 @@ void FScene::BackgroundTick()
 			}
 
 			ENQUEUE_RENDER_COMMAND(BackgroundTickRenderThread)(
-			[&RenderState = RenderState](FRHICommandListImmediate& RHICmdList) mutable
+				[&RenderState = RenderState](FRHICommandListImmediate& RHICmdList) mutable
 			{
 				RenderState.BackgroundTick(RHICmdList);
 			});
@@ -1581,7 +1583,7 @@ void FScene::BackgroundTick()
 	}
 }
 
-void FSceneRenderState::BackgroundTick(FRHICommandList& RHICmdList)
+void FSceneRenderState::BackgroundTick(FRHICommandListImmediate& RHICmdList)
 {
 	if (IrradianceCache->CurrentRevision != LightmapRenderer->GetCurrentRevision())
 	{
@@ -1671,13 +1673,13 @@ void FSceneRenderState::BackgroundTick(FRHICommandList& RHICmdList)
 		FPlatformAtomics::InterlockedExchange(&Percentage, IntPercentage);
 	}
 
-	LightmapRenderer->BackgroundTick();
+	LightmapRenderer->BackgroundTick(RHICmdList);
 
 	// If we're in background baking mode, schedule VLM work to be after surface lightmaps
 	const bool bIsViewportNonRealtime = !FGPULightmassModule::IsRealtimeOn();
 	if (!bIsViewportNonRealtime || (bIsViewportNonRealtime && bHaveFinishedSurfaceLightmaps))
 	{
-		VolumetricLightmapRenderer->BackgroundTick();
+		VolumetricLightmapRenderer->BackgroundTick(RHICmdList);
 	}
 }
 
diff --git a/Engine/Plugins/Experimental/GPULightmass/Source/GPULightmass/Private/Scene/Scene.h b/Engine/Plugins/Experimental/GPULightmass/Source/GPULightmass/Private/Scene/Scene.h
--- a/Engine/Plugins/Experimental/GPULightmass/Source/GPULightmass/Private/Scene/Scene.h
+++ b/Engine/Plugins/Experimental/GPULightmass/Source/GPULightmass/Private/Scene/Scene.h
@@ -102,7 +102,7 @@ public:
 	UGPULightmassSettings* Settings;
 
 	void RenderThreadInit(FRHICommandList& RHICmdList);
-	void BackgroundTick(FRHICommandList& RHICmdList);
+	void BackgroundTick(FRHICommandListImmediate& RHICmdList);
 
 	FRayTracingSceneRHIRef RayTracingScene;
 	FShaderResourceViewRHIRef RayTracingSceneSRV;
diff --git a/Engine/Plugins/Experimental/GPULightmass/Source/GPULightmass/Private/VolumetricLightmap.cpp b/Engine/Plugins/Experimental/GPULightmass/Source/GPULightmass/Private/VolumetricLightmap.cpp
--- a/Engine/Plugins/Experimental/GPULightmass/Source/GPULightmass/Private/VolumetricLightmap.cpp
+++ b/Engine/Plugins/Experimental/GPULightmass/Source/GPULightmass/Private/VolumetricLightmap.cpp
@@ -131,7 +131,7 @@ BEGIN_SHADER_PARAMETER_STRUCT(FSortBrickRequestsPassParameters, )
 	SHADER_PARAMETER_RDG_BUFFER_UAV(RWBuffer<uint>, BrickRequestSortedIndicesUAV2)
 END_SHADER_PARAMETER_STRUCT()
 	
-void FVolumetricLightmapRenderer::VoxelizeScene()
+void FVolumetricLightmapRenderer::VoxelizeScene(FRHICommandListImmediate& RHICmdList)
 {
 	FGlobalShaderMap* GlobalShaderMap = GetGlobalShaderMap(Scene->FeatureLevel);
 
@@ -146,7 +146,7 @@ void FVolumetricLightmapRenderer::VoxelizeScene()
 	ReleaseBrickData(VolumetricLightmapData.BrickData);
 	ReleaseBrickData(AccumulationBrickData);
 
-	FRHICommandListExecutor::GetImmediateCommandList().ImmediateFlush(EImmediateFlushType::FlushRHIThreadFlushResources);
+	RHICmdList.ImmediateFlush(EImmediateFlushType::FlushRHIThreadFlushResources);
 
 	VolumeMin = Scene->CombinedImportanceVolume.Min;
 
@@ -169,8 +169,6 @@ void FVolumetricLightmapRenderer::VoxelizeScene()
 
 	VoxelizationVolumeMips.Empty();
 
-	FRHICommandListImmediate& RHICmdList = FRHICommandListExecutor::GetImmediateCommandList();
-
 	{
 		FRDGBuilder GraphBuilder(RHICmdList);
 		ON_SCOPE_EXIT{ GraphBuilder.Execute(); };
@@ -716,7 +714,7 @@ struct FVolumetricLightmapBrickDataRDG
 	}
 };
 	
-void FVolumetricLightmapRenderer::BackgroundTick()
+void FVolumetricLightmapRenderer::BackgroundTick(FRHICommandListImmediate& RHICmdList)
 {
 	if (NumTotalBricks == 0)
 	{
@@ -730,7 +728,7 @@ void FVolumetricLightmapRenderer::BackgroundTick()
 
 	TRACE_CPUPROFILER_EVENT_SCOPE(FVolumetricLightmapRenderer::BackgroundTick);
 
-	FRDGBuilder GraphBuilder(FRHICommandListExecutor::GetImmediateCommandList());
+	FRDGBuilder GraphBuilder(RHICmdList);
 	FSceneUniformBuffer SceneUniforms {};
 	
 	if (IsRayTracingEnabled())
diff --git a/Engine/Plugins/Experimental/GPULightmass/Source/GPULightmass/Private/VolumetricLightmap.h b/Engine/Plugins/Experimental/GPULightmass/Source/GPULightmass/Private/VolumetricLightmap.h
--- a/Engine/Plugins/Experimental/GPULightmass/Source/GPULightmass/Private/VolumetricLightmap.h
+++ b/Engine/Plugins/Experimental/GPULightmass/Source/GPULightmass/Private/VolumetricLightmap.h
@@ -18,8 +18,8 @@ public:
 	
 	FVolumetricLightmapRenderer(FSceneRenderState* Scene);
 
-	void VoxelizeScene();
-	void BackgroundTick();
+	void VoxelizeScene(FRHICommandListImmediate& RHICmdList);
+	void BackgroundTick(FRHICommandListImmediate& RHICmdList);
 
 	FPrecomputedVolumetricLightmap* GetPrecomputedVolumetricLightmapForPreview();
 
diff --git a/Engine/Plugins/TextureGraph/Source/TextureGraphEngine/Device/FX/Device_FX.h b/Engine/Plugins/TextureGraph/Source/TextureGraphEngine/Device/FX/Device_FX.h
--- a/Engine/Plugins/TextureGraph/Source/TextureGraphEngine/Device/FX/Device_FX.h
+++ b/Engine/Plugins/TextureGraph/Source/TextureGraphEngine/Device/FX/Device_FX.h
@@ -104,7 +104,10 @@ public:
 	//////////////////////////////////////////////////////////////////////////
 	/// Inline functions
 	//////////////////////////////////////////////////////////////////////////
-	FORCEINLINE FRHICommandListImmediate& RHI() const { check(IsInRenderingThread()); return GRHICommandList.GetImmediateCommandList(); }
+	FORCEINLINE FRHICommandListImmediate& RHI() const
+	{
+		return FRHICommandListImmediate::Get();
+	}
 
 	//////////////////////////////////////////////////////////////////////////
 	/// Static functions
diff --git a/Engine/Plugins/TextureGraph/Source/TextureGraphEngine/Helper/Util.cpp b/Engine/Plugins/TextureGraph/Source/TextureGraphEngine/Helper/Util.cpp
--- a/Engine/Plugins/TextureGraph/Source/TextureGraphEngine/Helper/Util.cpp
+++ b/Engine/Plugins/TextureGraph/Source/TextureGraphEngine/Helper/Util.cpp
@@ -149,7 +149,7 @@ size_t Util::GetCurrentThreadId()
 
 FRHICommandListImmediate& Util::RHIImmediate()
 {
-	check(IsInRenderingThread()); return GRHICommandList.GetImmediateCommandList();
+	return FRHICommandListImmediate::Get();
 }
 
 //////////////////////////////////////////////////////////////////////////
diff --git a/Engine/Plugins/TextureGraph/Source/TextureGraphEngine/Helper/Util.h b/Engine/Plugins/TextureGraph/Source/TextureGraphEngine/Helper/Util.h
--- a/Engine/Plugins/TextureGraph/Source/TextureGraphEngine/Helper/Util.h
+++ b/Engine/Plugins/TextureGraph/Source/TextureGraphEngine/Helper/Util.h
@@ -87,6 +87,7 @@ struct Util
 	static UE_API void						OnThread(ENamedThreads::Type thread, TUniqueFunction<void()> Callback);
 	static UE_API size_t					GetCurrentThreadId();
 
+	UE_DEPRECATED(5.8, "Util::RHIImmediate() is deprecated. Use FRHICommandListImmediate::Get() instead.")
 	static UE_API FRHICommandListImmediate& RHIImmediate();
 
 	//////////////////////////////////////////////////////////////////////////
diff --git a/Engine/Source/Runtime/Engine/Private/GPUSkinVertexFactory.cpp b/Engine/Source/Runtime/Engine/Private/GPUSkinVertexFactory.cpp
--- a/Engine/Source/Runtime/Engine/Private/GPUSkinVertexFactory.cpp
+++ b/Engine/Source/Runtime/Engine/Private/GPUSkinVertexFactory.cpp
@@ -664,11 +664,6 @@ void GetGPUSkinAPEXClothVertexFactoryUniformShaderParameters(
 	UniformParameters.PreviousClothToLocal = ClothShaderData.GetClothToLocalForReading(true);
 }
 
-void FGPUBaseSkinVertexFactory::SetData(const FGPUSkinDataType* InData)
-{
-	SetData(FRHICommandListExecutor::GetImmediateCommandList(), InData);
-}
-
 void FGPUBaseSkinVertexFactory::SetData(FRHICommandListBase& RHICmdList, const FGPUSkinDataType* InData)
 {
 	check(InData);
diff --git a/Engine/Source/Runtime/Engine/Private/PrecomputedVolumetricLightmap.cpp b/Engine/Source/Runtime/Engine/Private/PrecomputedVolumetricLightmap.cpp
--- a/Engine/Source/Runtime/Engine/Private/PrecomputedVolumetricLightmap.cpp
+++ b/Engine/Source/Runtime/Engine/Private/PrecomputedVolumetricLightmap.cpp
@@ -334,12 +334,10 @@ ENGINE_API void FPrecomputedVolumetricLightmapData::ReleaseRHIForSubLevelResourc
 	}
 }
 
-ENGINE_API void FPrecomputedVolumetricLightmapData::HandleDataMovementInAtlas(int32 OldOffset, int32 NewOffset)
+ENGINE_API void FPrecomputedVolumetricLightmapData::HandleDataMovementInAtlas(FRHIComputeCommandList& RHICmdList, int32 OldOffset, int32 NewOffset)
 {
 	BrickDataBaseOffsetInAtlas = NewOffset;
 
-	FRHICommandListImmediate& RHICmdList = FRHICommandListExecutor::GetImmediateCommandList();
-
 	if (IndirectionTextureDimensions.GetMax() > 0)
 	{
 		const int32 PaddedBrickSize = BrickSize + 1;
@@ -1071,7 +1069,7 @@ void FVolumetricLightmapBrickAtlas::Insert(int32 Index, FPrecomputedVolumetricLi
 		{
 			if (Allocations[AllocationIndex].Data->IndirectionTextureDimensions.GetMax() > 0)
 			{
-				Allocations[AllocationIndex].Data->HandleDataMovementInAtlas(Allocations[AllocationIndex].StartOffset, BrickStartAllocation);
+				Allocations[AllocationIndex].Data->HandleDataMovementInAtlas(RHICmdList, Allocations[AllocationIndex].StartOffset, BrickStartAllocation);
 			}
 			BrickStartAllocation += Allocations[AllocationIndex].Size;
 		}
@@ -1109,7 +1107,7 @@ void FVolumetricLightmapBrickAtlas::Insert(int32 Index, FPrecomputedVolumetricLi
 			// Handle the sub level data movements
 			if (Allocations[AllocationIndex].Data->IndirectionTextureDimensions.GetMax() == 0)
 			{
-				Allocations[AllocationIndex].Data->HandleDataMovementInAtlas(Allocations[AllocationIndex].StartOffset, BrickStartAllocation);
+				Allocations[AllocationIndex].Data->HandleDataMovementInAtlas(RHICmdList, Allocations[AllocationIndex].StartOffset, BrickStartAllocation);
 			}
 			BrickStartAllocation += Allocations[AllocationIndex].Size;
 		}
@@ -1181,7 +1179,7 @@ void FVolumetricLightmapBrickAtlas::Remove(FPrecomputedVolumetricLightmapData* D
 				CopyDataIntoAtlas(RHICmdList, GetFeatureLevel(), Allocations[AllocationIndex].StartOffset, BrickStartAllocation, Allocations[AllocationIndex].Size, TextureSet, NewTextureSet);
 
 				NewAllocations.Add(Allocation{ Allocations[AllocationIndex].Data, Allocations[AllocationIndex].Size, BrickStartAllocation });
-				Allocations[AllocationIndex].Data->HandleDataMovementInAtlas(Allocations[AllocationIndex].StartOffset, BrickStartAllocation);
+				Allocations[AllocationIndex].Data->HandleDataMovementInAtlas(RHICmdList, Allocations[AllocationIndex].StartOffset, BrickStartAllocation);
 				BrickStartAllocation += Allocations[AllocationIndex].Size;
 			}
 		}
diff --git a/Engine/Source/Runtime/Engine/Private/Streaming/RenderAssetUpdate.cpp b/Engine/Source/Runtime/Engine/Private/Streaming/RenderAssetUpdate.cpp
--- a/Engine/Source/Runtime/Engine/Private/Streaming/RenderAssetUpdate.cpp
+++ b/Engine/Source/Runtime/Engine/Private/Streaming/RenderAssetUpdate.cpp
@@ -74,9 +74,9 @@ void SuspendRenderAssetStreaming()
 		{
 			FPlatformProcess::Sleep(RENDER_ASSET_STREAMING_SLEEP_DT);
 			ENQUEUE_RENDER_COMMAND(FlushCommand)(
-				[](FRHICommandList& RHICmdList)
+				[](FRHICommandListImmediate& RHICmdList)
 				{
-					GRHICommandList.GetImmediateCommandList().ImmediateFlush(EImmediateFlushType::FlushRHIThread);
+					RHICmdList.ImmediateFlush(EImmediateFlushType::FlushRHIThread);
 				});
 			FlushRenderingCommands();
 			
@@ -119,9 +119,9 @@ void SuspendRenderAssetStreaming()
 		{
 			// Ensure any pending render command executes.
 			ENQUEUE_RENDER_COMMAND(FlushCommand)(
-				[](FRHICommandList& RHICmdList)
+				[](FRHICommandListImmediate& RHICmdList)
 				{
-					GRHICommandList.GetImmediateCommandList().ImmediateFlush(EImmediateFlushType::FlushRHIThread);
+					RHICmdList.ImmediateFlush(EImmediateFlushType::FlushRHIThread);
 				});
 			FlushRenderingCommands();
 		}
diff --git a/Engine/Source/Runtime/Engine/Private/UnrealEngine.cpp b/Engine/Source/Runtime/Engine/Private/UnrealEngine.cpp
--- a/Engine/Source/Runtime/Engine/Private/UnrealEngine.cpp
+++ b/Engine/Source/Runtime/Engine/Private/UnrealEngine.cpp
@@ -2141,9 +2141,9 @@ void UEngine::Init(IEngineLoop* InEngineLoop)
 			{
 				// Flush render and RHI commands in case there are some UStreamableRenderAssets with pending RenderAssetUpdate RHI tasks
 				ENQUEUE_RENDER_COMMAND(FlushCommand)(
-					[](FRHICommandList& RHICmdList)
+					[](FRHICommandListImmediate& RHICmdList)
 					{
-						GRHICommandList.GetImmediateCommandList().ImmediateFlush(EImmediateFlushType::FlushRHIThread);
+						RHICmdList.ImmediateFlush(EImmediateFlushType::FlushRHIThread);
 					});
 				FlushRenderingCommands();
 
@@ -16457,11 +16457,11 @@ void UEngine::TrimMemory()
 	// so that any pending frees that depend on the GPU will be processed.  Otherwise a whole map's worth of GPU memory
 	// may be unavailable to load the next one.
 	ENQUEUE_RENDER_COMMAND(FlushCommand)(
-		[](FRHICommandList& RHICmdList)
+		[](FRHICommandListImmediate& RHICmdList)
 		{
-			GRHICommandList.GetImmediateCommandList().ImmediateFlush(EImmediateFlushType::FlushRHIThread);
+			RHICmdList.ImmediateFlush(EImmediateFlushType::FlushRHIThread);
 			RHIFlushResources();
-			GRHICommandList.GetImmediateCommandList().ImmediateFlush(EImmediateFlushType::FlushRHIThreadFlushResources);
+			RHICmdList.ImmediateFlush(EImmediateFlushType::FlushRHIThreadFlushResources);
 		});
 	FlushRenderingCommands();
 
diff --git a/Engine/Source/Runtime/Engine/Public/GPUSkinVertexFactory.h b/Engine/Source/Runtime/Engine/Public/GPUSkinVertexFactory.h
--- a/Engine/Source/Runtime/Engine/Public/GPUSkinVertexFactory.h
+++ b/Engine/Source/Runtime/Engine/Public/GPUSkinVertexFactory.h
@@ -352,14 +352,6 @@ public:
 		return ShaderData;
 	}
 
-	/**
-	* An implementation of the interface used by TSynchronizedResource to
-	* update the resource with new data from the game thread.
-	* @param	InData - new stream component data
-	*/
-	UE_DEPRECATED(5.3, "Use SetData with a command list.")
-	void SetData(const FGPUSkinDataType* InData);
-
 	virtual void SetData(FRHICommandListBase& RHICmdList, const FGPUSkinDataType* InData);
 
 	uint32 GetBoneBufferSize() const
diff --git a/Engine/Source/Runtime/Engine/Public/PrecomputedVolumetricLightmap.h b/Engine/Source/Runtime/Engine/Public/PrecomputedVolumetricLightmap.h
--- a/Engine/Source/Runtime/Engine/Public/PrecomputedVolumetricLightmap.h
+++ b/Engine/Source/Runtime/Engine/Public/PrecomputedVolumetricLightmap.h
@@ -165,7 +165,13 @@ public:
 	ENGINE_API void InitRHIForSubLevelResources();
 	ENGINE_API void ReleaseRHIForSubLevelResources();
 
-	ENGINE_API void HandleDataMovementInAtlas(int32 OldOffset, int32 NewOffset);
+	UE_DEPRECATED(5.8, "HandleDataMovementInAtlas now requires a command list.")
+	void HandleDataMovementInAtlas(int32 OldOffset, int32 NewOffset)
+	{
+		HandleDataMovementInAtlas(FRHICommandListImmediate::Get(), OldOffset, NewOffset);
+	}
+
+	ENGINE_API void HandleDataMovementInAtlas(FRHIComputeCommandList& RHICmdList, int32 OldOffset, int32 NewOffset);
 	ENGINE_API void AddToSceneData(FPrecomputedVolumetricLightmapData* SceneData);
 	ENGINE_API void RemoveFromSceneData(FPrecomputedVolumetricLightmapData* SceneData, int32 PersistentLevelBrickDataBaseOffset);
 
diff --git a/Engine/Source/Runtime/RHI/Private/DynamicRHI.cpp b/Engine/Source/Runtime/RHI/Private/DynamicRHI.cpp
--- a/Engine/Source/Runtime/RHI/Private/DynamicRHI.cpp
+++ b/Engine/Source/Runtime/RHI/Private/DynamicRHI.cpp
@@ -323,13 +323,13 @@ void RHIInit(bool bHasEditorToken)
 				GDynamicRHI->Init();
 
 				// Validation of contexts.
-				GRHICommandList.GetImmediateCommandList().GetContext();
+				FRHICommandListImmediate::Get().GetContext();
 				check(GIsRHIInitialized);
 
 				// Set default GPU mask to all GPUs. This is necessary to ensure that any commands
 				// that create and initialize resources are executed on all GPUs. Scene rendering
 				// will restrict itself to a subset of GPUs as needed.
-				GRHICommandList.GetImmediateCommandList().SetGPUMask(FRHIGPUMask::All());
+				FRHICommandListImmediate::Get().SetGPUMask(FRHIGPUMask::All());
 
 				FString FeatureLevelString;
 				GetFeatureLevelName(GMaxRHIFeatureLevel, FeatureLevelString);
diff --git a/Engine/Source/Runtime/RenderCore/Private/ShaderPipelineCache.cpp b/Engine/Source/Runtime/RenderCore/Private/ShaderPipelineCache.cpp
--- a/Engine/Source/Runtime/RenderCore/Private/ShaderPipelineCache.cpp
+++ b/Engine/Source/Runtime/RenderCore/Private/ShaderPipelineCache.cpp
@@ -427,12 +427,12 @@ public:
 	bool Precompile(FRHICommandListImmediate& RHICmdList, EShaderPlatform Platform, FPipelineCacheFileFormatPSO const& PSO);
 	void PreparePipelineBatch(TDoubleLinkedList<FPipelineCacheFileFormatPSORead*>& PipelineBatch);
 	bool ReadyForPrecompile();
-	void PrecompilePipelineBatch();
+	void PrecompilePipelineBatch(FRHICommandListImmediate& RHICmdList);
 	bool ReadyForNextBatch() const;
 	bool ReadyForAutoSave() const;
 	void PollShutdownItems();
 	void Flush();
-	void Tick(float DeltaTime);
+	void Tick(FRHICommandListImmediate& RHICmdList, float DeltaTime);
 	bool OnShaderLibraryStateChanged(FShaderPipelineCache::ELibraryState State, EShaderPlatform Platform, FString const& Name, int32 ComponentID);
 private:
 	FString PSOCacheKey;
@@ -1483,7 +1483,7 @@ bool FShaderPipelineCacheTask::ReadyForPrecompile()
 		&& PipelineStateCache::GetNumActivePipelinePrecompileTasks() < GMinPipelinePrecompileTasksInFlightThreshold;
 }
 
-void FShaderPipelineCacheTask::PrecompilePipelineBatch()
+void FShaderPipelineCacheTask::PrecompilePipelineBatch(FRHICommandListImmediate& RHICmdList)
 {
 	INC_DWORD_STAT(STAT_PreCompileBatchTotal);
 	INC_DWORD_STAT(STAT_PreCompileBatchNum);
@@ -1498,8 +1498,6 @@ void FShaderPipelineCacheTask::PrecompilePipelineBatch()
 		
 		check(CompileTask.ReadRequests && CompileTask.ReadRequests->PollExternalReadDependencies());
 		
-		FRHICommandListImmediate& RHICmdList = GRHICommandList.GetImmediateCommandList();
-		
 		uint32 PSOHash = GetTypeHash(CompileTask.PSO);
 		UE_LOG(LogRHI, Verbose, TEXT("Precompiling PSO %u (type %d) (%d/%d)"), PSOHash, int(CompileTask.PSO.Type), i+1, NumToPrecompile);
 		Precompile(RHICmdList, GMaxRHIShaderPlatform, CompileTask.PSO);
@@ -1859,7 +1857,7 @@ void FShaderPipelineCache::Tick(FRHICommandListImmediate& RHICmdList, float Delt
 
 	if(CurrentCacheTask)
 	{
-		CurrentCacheTask->Tick(DeltaTime);
+		CurrentCacheTask->Tick(RHICmdList, DeltaTime);
 	}
 }
 
@@ -1869,7 +1867,7 @@ void FShaderPipelineCache::Tick(float DeltaTime)
 	Tick(FRHICommandListImmediate::Get(), DeltaTime);
 }
 
-void FShaderPipelineCacheTask::Tick(float DeltaTime)
+void FShaderPipelineCacheTask::Tick(FRHICommandListImmediate& RHICmdList, float DeltaTime)
 {
 	if (PrecompileStartTime > 0.0)
 	{
@@ -1888,7 +1886,7 @@ void FShaderPipelineCacheTask::Tick(float DeltaTime)
 		
 		uint32 Start = FPlatformTime::Cycles();
 
-		PrecompilePipelineBatch();
+		PrecompilePipelineBatch(RHICmdList);
 
 		uint32 End = FPlatformTime::Cycles();
 
