[GPULM] Fix several crashes by moving TLAS and SBT builds to RDG.

#rnx
#tests GPULM tests

#Synced-CL 48575655

diff --git a/Engine/Plugins/Experimental/GPULightmass/Source/GPULightmass/Private/LightmapRenderer.cpp b/Engine/Plugins/Experimental/GPULightmass/Source/GPULightmass/Private/LightmapRenderer.cpp
--- a/Engine/Plugins/Experimental/GPULightmass/Source/GPULightmass/Private/LightmapRenderer.cpp
+++ b/Engine/Plugins/Experimental/GPULightmass/Source/GPULightmass/Private/LightmapRenderer.cpp
@@ -855,6 +855,11 @@ FCachedRayTracingSceneData::~FCachedRayTracingSceneData()
 	FRHICommandListImmediate::Get().EnqueueLambda([OwnedRayTracingInstanceTransforms = MoveTemp(OwnedRayTracingInstanceTransforms)](FRHICommandList&){});
 }
 
+BEGIN_SHADER_PARAMETER_STRUCT(FRayTracingSceneBuildPassParams, )
+	SHADER_PARAMETER_STRUCT_REF(FViewUniformShaderParameters, View)
+	SHADER_PARAMETER_RDG_UNIFORM_BUFFER(FSceneUniformParameters, Scene)
+END_SHADER_PARAMETER_STRUCT()
+
 bool FSceneRenderState::SetupRayTracingScene(FRDGBuilder& GraphBuilder, FSceneUniformBuffer& SceneUniforms, int32 LODIndex)
 {
 	TRACE_CPUPROFILER_EVENT_SCOPE(SetupRayTracingScene);
@@ -979,9 +984,13 @@ bool FSceneRenderState::SetupRayTracingScene(FRDGBuilder& GraphBuilder, FSceneUn
 	{
 		TRACE_CPUPROFILER_EVENT_SCOPE(RayTracingScene);
 
+		FRayTracingSceneBuildPassParams* PassParameters = GraphBuilder.AllocParameters<FRayTracingSceneBuildPassParams>();
+		PassParameters->View = ReferenceView->ViewUniformBuffer;
+		PassParameters->Scene = SceneUniforms.GetBuffer(GraphBuilder);
+
 		SCOPED_DRAW_EVENTF(GraphBuilder.RHICmdList, GPULightmassUpdateRayTracingScene, TEXT("GPULightmass UpdateRayTracingScene %d Instances"), StaticMeshInstanceRenderStates.Elements.Num());
 
-		TArray<FRayTracingGeometryInstance, SceneRenderingAllocator> RayTracingGeometryInstances;
+		TArray<FRayTracingGeometryInstance, SceneRenderingAllocator>& RayTracingGeometryInstances = GraphBuilder.AllocArray<FRayTracingGeometryInstance>();
 		if (LODIndex != INDEX_NONE)
 		{
 			RayTracingGeometryInstances.Append(CachedRayTracingScene->RayTracingGeometryInstancesPerLOD[LODIndex]);
@@ -1048,7 +1057,7 @@ bool FSceneRenderState::SetupRayTracingScene(FRDGBuilder& GraphBuilder, FSceneUn
 							Landscape.SectionRayTracingStates[SubSectionIdx]->Geometry.SetInitializer(GeometryInitializer);
 							Landscape.SectionRayTracingStates[SubSectionIdx]->Geometry.InitResource(GraphBuilder.RHICmdList);
 
-							FRayTracingDynamicGeometryUpdateManager DynamicGeometryUpdateManager;
+							FRayTracingDynamicGeometryUpdateManager& DynamicGeometryUpdateManager = *GraphBuilder.AllocObject<FRayTracingDynamicGeometryUpdateManager>();
 
 							TArray<FMeshBatch> MeshBatches = Landscape.GetMeshBatchesForGBufferRendering(0);
 
@@ -1119,7 +1128,7 @@ bool FSceneRenderState::SetupRayTracingScene(FRDGBuilder& GraphBuilder, FSceneUn
 
 						FRayTracingGeometryInstance& RayTracingInstance = RayTracingGeometryInstances[InstanceIndex];
 						RayTracingInstance.GeometryRHI = Landscape.SectionRayTracingStates[SubSectionIdx]->Geometry.GetRHI();
-						RayTracingInstance.Transforms = MakeArrayView(LandscapeTransforms.Emplace_GetRef(MakeUnique<FMatrix>(Landscape.LocalToWorld)).Get(), 1);
+						RayTracingInstance.Transforms = MakeArrayView(GraphBuilder.AllocObject<FMatrix>(Landscape.LocalToWorld), 1);
 						RayTracingInstance.NumTransforms = 1;
 						RayTracingInstance.InstanceContributionToHitGroupIndex = SBTAllocation->GetInstanceContributionToHitGroupIndex(ERayTracingShaderBindingLayer::Base);
 						RayTracingInstance.DefaultUserData = (uint32)(StaticMeshInstanceRenderStates.Elements.Num() + InstanceGroupRenderStates.Elements.Num() + LandscapeIndex);
@@ -1205,43 +1214,57 @@ bool FSceneRenderState::SetupRayTracingScene(FRDGBuilder& GraphBuilder, FSceneUn
 				HitGroupContributionsBuffer.Initialize(GraphBuilder.RHICmdList, TEXT("LightmassRayTracingHitGroupBuffer"), 4, SceneInitializer.MaxNumInstances);
 			}
 
-			RayTracingInstanceBufferBuilder.FillRayTracingInstanceUploadBuffer(GraphBuilder.RHICmdList);
-			RayTracingInstanceBufferBuilder.FillAccelerationStructureAddressesBuffer(GraphBuilder.RHICmdList);
+			GraphBuilder.AddPass(
+				RDG_EVENT_NAME("GPULM::RayTracingBuildScene"),
+				PassParameters,
+				ERDGPassFlags::Compute | ERDGPassFlags::NeverCull,
+				[
+					RayTracingScene = RayTracingScene,
+					RayTracingSceneBuffer = RayTracingSceneBuffer,
+					ScratchBuffer = ScratchBuffer,
+					MaxNumInstances = SceneInitializer.MaxNumInstances,
+					InstanceBufferUAV = InstanceBuffer.UAV,
+					InstanceBufferRHI = InstanceBuffer.Buffer,
+					bRequiresSeparateHitGroupContributionsBuffer = GRHIGlobals.RayTracing.RequiresSeparateHitGroupContributionsBuffer,
+					HitGroupContributionsBufferUAV = HitGroupContributionsBuffer.UAV,
+					HitGroupContributionsBufferRHI = HitGroupContributionsBuffer.Buffer,
+					RayTracingInstanceBufferBuilder = MoveTemp(RayTracingInstanceBufferBuilder)](FRHICommandListImmediate& RHICmdList) mutable
+				{
+					RayTracingInstanceBufferBuilder.FillRayTracingInstanceUploadBuffer(RHICmdList);
+					RayTracingInstanceBufferBuilder.FillAccelerationStructureAddressesBuffer(RHICmdList);
 
-			RayTracingInstanceBufferBuilder.BuildRayTracingInstanceBuffer(
-				GraphBuilder.RHICmdList,
-				nullptr,
-				nullptr,
-				nullptr,
-				InstanceBuffer.UAV,
-				GRHIGlobals.RayTracing.RequiresSeparateHitGroupContributionsBuffer ? HitGroupContributionsBuffer.UAV : nullptr,
-				SceneInitializer.MaxNumInstances,
-				/*bCompactOutput*/ false,
-				nullptr,
-				0,
-				nullptr);
+					RayTracingInstanceBufferBuilder.BuildRayTracingInstanceBuffer(
+						RHICmdList,
+						nullptr,
+						nullptr,
+						nullptr,
+						InstanceBufferUAV,
+						bRequiresSeparateHitGroupContributionsBuffer ? HitGroupContributionsBufferUAV : nullptr,
+						MaxNumInstances,
+						false,
+						nullptr,
+						0,
+						nullptr);
 
-			GraphBuilder.RHICmdList.BindAccelerationStructureMemory(RayTracingScene, RayTracingSceneBuffer, 0);
+					FRayTracingSceneBuildParams BuildParams;
+					BuildParams.Scene = RayTracingScene;
+					BuildParams.ScratchBuffer = ScratchBuffer;
+					BuildParams.ScratchBufferOffset = 0;
+					BuildParams.InstanceBuffer = InstanceBufferRHI;
+					BuildParams.InstanceBufferOffset = 0;
 
-			{
-				FRayTracingSceneBuildParams BuildParams;
-				BuildParams.Scene = RayTracingScene;
-				BuildParams.ScratchBuffer = ScratchBuffer;
-				BuildParams.ScratchBufferOffset = 0;
-				BuildParams.InstanceBuffer = InstanceBuffer.Buffer;
-				BuildParams.InstanceBufferOffset = 0;
-				
-				if(GRHIGlobals.RayTracing.RequiresSeparateHitGroupContributionsBuffer)
-				{
-					BuildParams.HitGroupContributionsBuffer = HitGroupContributionsBuffer.Buffer;
-					BuildParams.HitGroupContributionsBufferOffset = 0;
-				}
-				
-				BuildParams.ReferencedGeometries = RayTracingInstanceBufferBuilder.GetReferencedGeometries();
-				BuildParams.NumInstances = RayTracingInstanceBufferBuilder.GetMaxNumInstances();
+					if (bRequiresSeparateHitGroupContributionsBuffer)
+					{
+						BuildParams.HitGroupContributionsBuffer = HitGroupContributionsBufferRHI;
+						BuildParams.HitGroupContributionsBufferOffset = 0;
+					}
 
-				GraphBuilder.RHICmdList.BuildAccelerationStructure(BuildParams);
-			}
+					BuildParams.ReferencedGeometries = RayTracingInstanceBufferBuilder.GetReferencedGeometries();
+					BuildParams.NumInstances = RayTracingInstanceBufferBuilder.GetMaxNumInstances();
+
+					RHICmdList.BindAccelerationStructureMemory(RayTracingScene, RayTracingSceneBuffer, 0);
+					RHICmdList.BuildAccelerationStructure(BuildParams);
+				});
 
 			// Move LandscapeTransforms to RHIThread to extend its lifetime until RHI cmd execution
 			GraphBuilder.RHICmdList.EnqueueLambda([LandscapeTransforms = MoveTemp(LandscapeTransforms)](FRHICommandList&){});
@@ -1362,48 +1385,53 @@ bool FSceneRenderState::SetupRayTracingScene(FRDGBuilder& GraphBuilder, FSceneUn
 						ERayTracingLocalShaderBindingType::Transient);
 				}
 
-				{
-					uint32 NumTotalBindings = 0;
-
-					for (const FRayTracingLocalShaderBindingWriter::FChunk* Chunk = BindingWriter->GetFirstChunk(); Chunk; Chunk = Chunk->Next)
+				GraphBuilder.AddPass(
+					RDG_EVENT_NAME("GPULM::SetRayTracingBindings"),
+					PassParameters,
+					ERDGPassFlags::Compute | ERDGPassFlags::NeverCull,
+					[BindingWriter = MoveTemp(BindingWriter), RayTracingPipelineState = RayTracingPipelineState, SBT = SBT](FRHICommandListImmediate& RHICmdList) mutable
 					{
-						NumTotalBindings += Chunk->Num;
-					}
-
-					FConcurrentLinearBulkObjectAllocator Allocator;
-
-					const uint32 MergedBindingsSize = sizeof(FRayTracingLocalShaderBindings) * NumTotalBindings;
-					FRayTracingLocalShaderBindings* MergedBindings = static_cast<FRayTracingLocalShaderBindings*>(GraphBuilder.RHICmdList.Bypass()
-						?              Allocator.Malloc(MergedBindingsSize, alignof(FRayTracingLocalShaderBindings))
-						: GraphBuilder.RHICmdList.Alloc(MergedBindingsSize, alignof(FRayTracingLocalShaderBindings)));
+						uint32 NumTotalBindings = 0;
 
-					uint32 MergedBindingIndex = 0;
-					for (const FRayTracingLocalShaderBindingWriter::FChunk* Chunk = BindingWriter->GetFirstChunk(); Chunk; Chunk = Chunk->Next)
-					{
-						const uint32 Num = Chunk->Num;
-						for (uint32_t i = 0; i < Num; ++i)
+						for (const FRayTracingLocalShaderBindingWriter::FChunk* Chunk = BindingWriter->GetFirstChunk(); Chunk; Chunk = Chunk->Next)
 						{
-							MergedBindings[MergedBindingIndex] = Chunk->Bindings[i];
-							MergedBindingIndex++;
+							NumTotalBindings += Chunk->Num;
 						}
-					}
 
-					const bool bCopyDataToInlineStorage = false; // Storage is already allocated from RHICmdList, no extra copy necessary
-					GraphBuilder.RHICmdList.SetRayTracingHitGroups(
-						SBT,
-						RayTracingPipelineState,
-						NumTotalBindings, MergedBindings,
-						bCopyDataToInlineStorage);
-				}
+						FConcurrentLinearBulkObjectAllocator Allocator;
 
-				// there is only one miss shader, so it must be at index 0 by definition
-				GraphBuilder.RHICmdList.SetRayTracingMissShader(SBT, 0, RayTracingPipelineState, 0 /* ShaderIndexInPipeline */, 0, nullptr, 0);
-				GraphBuilder.RHICmdList.CommitShaderBindingTable(SBT);
+						const uint32 MergedBindingsSize = sizeof(FRayTracingLocalShaderBindings) * NumTotalBindings;
+						FRayTracingLocalShaderBindings* MergedBindings = static_cast<FRayTracingLocalShaderBindings*>(RHICmdList.Bypass()
+							? Allocator.Malloc(MergedBindingsSize, alignof(FRayTracingLocalShaderBindings))
+							: RHICmdList.Alloc(MergedBindingsSize, alignof(FRayTracingLocalShaderBindings)));
 
-				// Move the ray tracing binding container ownership to the command list, so that memory will be
-				// released on the RHI thread timeline, after the commands that reference it are processed.
-				// TUniquePtr<> auto destructs when exiting the lambda
-				GraphBuilder.RHICmdList.EnqueueLambda([BindingWriter = MoveTemp(BindingWriter)](FRHICommandListImmediate&) {});
+						uint32 MergedBindingIndex = 0;
+						for (const FRayTracingLocalShaderBindingWriter::FChunk* Chunk = BindingWriter->GetFirstChunk(); Chunk; Chunk = Chunk->Next)
+						{
+							const uint32 Num = Chunk->Num;
+							for (uint32_t i = 0; i < Num; ++i)
+							{
+								MergedBindings[MergedBindingIndex] = Chunk->Bindings[i];
+								MergedBindingIndex++;
+							}
+						}
+
+						const bool bCopyDataToInlineStorage = false; // Storage is already allocated from RHICmdList, no extra copy necessary
+						RHICmdList.SetRayTracingHitGroups(
+							SBT,
+							RayTracingPipelineState,
+							NumTotalBindings, MergedBindings,
+							bCopyDataToInlineStorage);
+
+						// there is only one miss shader, so it must be at index 0 by definition
+						RHICmdList.SetRayTracingMissShader(SBT, 0, RayTracingPipelineState, 0 /* ShaderIndexInPipeline */, 0, nullptr, 0);
+						RHICmdList.CommitShaderBindingTable(SBT);
+
+						// Move the ray tracing binding container ownership to the command list, so that memory will be
+						// released on the RHI thread timeline, after the commands that reference it are processed.
+						// TUniquePtr<> auto destructs when exiting the lambda
+						RHICmdList.EnqueueLambda([BindingWriter = MoveTemp(BindingWriter)](FRHICommandListImmediate&) {});
+					});
 			}
 		}
 	}
