[HWRT] Fix error due to missing #include "RayTracingGeometryInstance.h" in some files.

#Synced-CL 50355052

diff --git a/Engine/Plugins/Experimental/GPULightmass/Source/GPULightmass/Private/Scene/Scene.h b/Engine/Plugins/Experimental/GPULightmass/Source/GPULightmass/Private/Scene/Scene.h
--- a/Engine/Plugins/Experimental/GPULightmass/Source/GPULightmass/Private/Scene/Scene.h
+++ b/Engine/Plugins/Experimental/GPULightmass/Source/GPULightmass/Private/Scene/Scene.h
@@ -8,6 +8,7 @@
 #include "Scene/Landscape.h"
 #include "MeshPassProcessor.h"
 #include "RayTracingMeshDrawCommands.h"
+#include "RayTracingGeometryInstance.h"
 #include "RayTracing/RayTracingShaderBindingTable.h"
 #include "IrradianceCaching.h"
 #include "GPULightmassSettings.h"
diff --git a/Engine/Source/Runtime/Renderer/Private/RayTracing/RayTracingScene.h b/Engine/Source/Runtime/Renderer/Private/RayTracing/RayTracingScene.h
--- a/Engine/Source/Runtime/Renderer/Private/RayTracing/RayTracingScene.h
+++ b/Engine/Source/Runtime/Renderer/Private/RayTracing/RayTracingScene.h
@@ -18,6 +18,7 @@
 #include "MeshPassProcessor.h"
 #include "RayTracingShaderBindingTable.h"
 #include "RayTracingInstanceBufferUtil.h"
+#include "RayTracingGeometryInstance.h"
 #include "RayTracingDebugTypes.h"
 
 class FGPUScene;
